Restore deleted (in two places) LevelSnapshots plugin

diff --git a/Engine/Plugins/VirtualProduction/LevelSnapshots/Source/LevelShapshotFilters/Private/LevelSnapshotFiltersModule.cpp b/Engine/Plugins/VirtualProduction/LevelSnapshots/Source/LevelShapshotFilters/Private/LevelSnapshotFiltersModule.cpp
new file mode 100644
--- /dev/null
+++ b/Engine/Plugins/VirtualProduction/LevelSnapshots/Source/LevelShapshotFilters/Private/LevelSnapshotFiltersModule.cpp
@@ -0,0 +1,20 @@
+// Copyright Epic Games, Inc. All Rights Reserved.
+
+#include "LevelSnapshotFiltersModule.h"
+
+#define LOCTEXT_NAMESPACE "FLevelSnapshotFiltersModule"
+
+void FLevelSnapshotFiltersModule::StartupModule()
+{
+	// This code will execute after your module is loaded into memory; the exact timing is specified in the .uplugin file per-module
+}
+
+void FLevelSnapshotFiltersModule::ShutdownModule()
+{
+	// This function may be called during shutdown to clean up your module.  For modules that support dynamic reloading,
+	// we call this function before unloading the module.
+}
+
+#undef LOCTEXT_NAMESPACE
+	
+IMPLEMENT_MODULE(FLevelSnapshotFiltersModule, LevelSnapshotFilters)
\ No newline at end of file
diff --git a/Engine/Plugins/VirtualProduction/LevelSnapshots/Source/LevelShapshotFilters/Public/LevelSnapshotFiltersModule.h b/Engine/Plugins/VirtualProduction/LevelSnapshots/Source/LevelShapshotFilters/Public/LevelSnapshotFiltersModule.h
new file mode 100644
--- /dev/null
+++ b/Engine/Plugins/VirtualProduction/LevelSnapshots/Source/LevelShapshotFilters/Public/LevelSnapshotFiltersModule.h
@@ -0,0 +1,15 @@
+// Copyright Epic Games, Inc. All Rights Reserved.
+
+#pragma once
+
+#include "CoreMinimal.h"
+#include "Modules/ModuleManager.h"
+
+class FLevelSnapshotFiltersModule : public IModuleInterface
+{
+public:
+
+	/** IModuleInterface implementation */
+	virtual void StartupModule() override;
+	virtual void ShutdownModule() override;
+};
diff --git a/Engine/Plugins/VirtualProduction/LevelSnapshots/Source/LevelSnapshots/Public/LevelSnapshotsLog.h b/Engine/Plugins/VirtualProduction/LevelSnapshots/Source/LevelSnapshots/Public/LevelSnapshotsLog.h
new file mode 100644
--- /dev/null
+++ b/Engine/Plugins/VirtualProduction/LevelSnapshots/Source/LevelSnapshots/Public/LevelSnapshotsLog.h
@@ -0,0 +1,7 @@
+// Copyright Epic Games, Inc. All Rights Reserved.
+
+#pragma once
+
+#include "CoreMinimal.h"
+
+DECLARE_LOG_CATEGORY_CLASS(LogLevelSnapshots,   Log, All);
diff --git a/Engine/Plugins/VirtualProduction/LevelSnapshots/Source/LevelSnapshotsEditor/Private/AssetTypeActions/AssetTypeActions_LevelSnapshot.h b/Engine/Plugins/VirtualProduction/LevelSnapshots/Source/LevelSnapshotsEditor/Private/AssetTypeActions/AssetTypeActions_LevelSnapshot.h
new file mode 100644
--- /dev/null
+++ b/Engine/Plugins/VirtualProduction/LevelSnapshots/Source/LevelSnapshotsEditor/Private/AssetTypeActions/AssetTypeActions_LevelSnapshot.h
@@ -0,0 +1,72 @@
+// Copyright Epic Games, Inc. All Rights Reserved.
+
+#pragma once
+
+#include "AssetTypeActions_Base.h"
+
+class FAssetTypeActions_LevelSnapshot : public FAssetTypeActions_Base
+{
+
+public:
+	// ~ Begin IAssetTypeActions Interface
+	
+	/**
+	 * @brief 
+	 * @return Whether or not this asset has a live filter
+	 */
+	virtual bool CanFilter() override
+	{
+		return true;
+	}
+	
+	/**
+	 * @brief 
+	 * @return Which category the asset will appear in when creating a new asset in the content browser
+	 */
+	virtual uint32 GetCategories() override
+	{
+		return EAssetTypeCategories::Misc;
+	}
+	
+	/**
+	 * @brief 
+	 * @return The color on the asset image container
+	 */
+	virtual FColor GetTypeColor() const override
+	{
+		return FColor(238, 181, 235, 255);
+	}
+	
+	/**
+	 * @brief 
+	 * @param InObjects A list of objects selected in the content browser
+	 * @return Whether or not the asset has context actions (required to return true in order to use registered actions)
+	 */
+	virtual bool HasActions(const TArray<UObject*>& InObjects) const override
+	{
+		return true;
+	}
+
+	// These are defined in source - GetName requires a defined LOCTEXT namespace and GetSupportedClass has an external include
+	
+	/**
+	 * @brief
+	 * @return The display name for the asset type
+	 */
+	virtual FText GetName() const override;
+
+	/**
+	 * @brief
+	 * @return The asset class these actions will belong to
+	 */
+	virtual UClass* GetSupportedClass() const override;
+	
+	/**
+	 * @brief Called when an item in the content browser is right-clicked. This determines what actions will be shown if HasActions() returns true.
+	 * @param InObjects A list of objects selected in the content browser
+	 * @param MenuBuilder A reference to an instance of the class which adds menu items to context menus
+	 */
+	virtual void GetActions(const TArray<UObject*>& InObjects, FMenuBuilder& MenuBuilder) override;
+
+	// ~ End IAssetTypeActions Interface
+};
diff --git a/Engine/Plugins/VirtualProduction/LevelSnapshots/Source/LevelSnapshotsEditor/Private/Data/UClassMetaDataDefinitions.h b/Engine/Plugins/VirtualProduction/LevelSnapshots/Source/LevelSnapshotsEditor/Private/Data/UClassMetaDataDefinitions.h
new file mode 100644
--- /dev/null
+++ b/Engine/Plugins/VirtualProduction/LevelSnapshots/Source/LevelSnapshotsEditor/Private/Data/UClassMetaDataDefinitions.h
@@ -0,0 +1,20 @@
+// Copyright Epic Games, Inc. All Rights Reserved.
+
+#pragma once
+
+#include "CoreMinimal.h"
+
+/* Special meta data used to internally label filter classes. These tags are uses by UFavoriteFilterContainer.
+ * Usage: UCLASS(meta = (CommonSnapshotFilter)) or UCLASS(meta = (InternalSnapshotFilter)).
+ */
+namespace UClassMetaDataDefinitions
+{
+	/* This filter is shown directly in the favorite filter menu opposed to the native and blueprint filter submenus.
+	 * Incompatible with InternalSnapshotFilter.
+	 */
+	const FName CommonSnapshotFilter("CommonSnapshotFilter");
+	/* This filter does not show up in the favorite filter menu.
+	 * Incompatible with CommonSnapshotFilter.
+	 */
+	const FName InternalSnapshotFilter("InternalSnapshotFilter");
+}
\ No newline at end of file
diff --git a/Engine/Plugins/VirtualProduction/LevelSnapshots/Source/LevelSnapshotsEditor/Private/Factories/LevelSnapshotFilterFactory.h b/Engine/Plugins/VirtualProduction/LevelSnapshots/Source/LevelSnapshotsEditor/Private/Factories/LevelSnapshotFilterFactory.h
new file mode 100644
--- /dev/null
+++ b/Engine/Plugins/VirtualProduction/LevelSnapshots/Source/LevelSnapshotsEditor/Private/Factories/LevelSnapshotFilterFactory.h
@@ -0,0 +1,27 @@
+// Copyright Epic Games, Inc. All Rights Reserved.
+
+#pragma once
+
+#include "CoreMinimal.h"
+#include "Factories/Factory.h"
+
+#include "LevelSnapshotFilterFactory.generated.h"
+
+UCLASS(hidecategories = Object)
+class ULevelSnapshotBlueprintFilterFactory : public UFactory
+{
+	GENERATED_BODY()
+public:
+	ULevelSnapshotBlueprintFilterFactory();
+	
+	// The parent class of the created blueprint
+    UPROPERTY(VisibleAnywhere, Category="LevelSnapshotBlueprintFilterFactory", meta=(AllowAbstract = "", BlueprintBaseOnly = ""))
+    TSubclassOf<UObject> ParentClass;
+	
+	//~ Begin UFactory Interface
+	virtual FText GetDisplayName() const override;
+	virtual UObject* FactoryCreateNew(UClass* Class, UObject* InParent, FName Name, EObjectFlags Flags, UObject* Context, FFeedbackContext* Warn) override;
+	virtual uint32 GetMenuCategories() const override;
+	//~ Begin UFactory Interface
+	
+};
\ No newline at end of file
diff --git a/Engine/Plugins/VirtualProduction/LevelSnapshots/Source/LevelSnapshotsEditor/Private/LevelSnapshotsEditorStyle.h b/Engine/Plugins/VirtualProduction/LevelSnapshots/Source/LevelSnapshotsEditor/Private/LevelSnapshotsEditorStyle.h
new file mode 100644
--- /dev/null
+++ b/Engine/Plugins/VirtualProduction/LevelSnapshots/Source/LevelSnapshotsEditor/Private/LevelSnapshotsEditorStyle.h
@@ -0,0 +1,37 @@
+// Copyright Epic Games, Inc. All Rights Reserved.
+
+#pragma once
+
+#include "Styling/ISlateStyle.h"
+#include "Styling/SlateStyle.h"
+
+class FLevelSnapshotsEditorStyle
+{
+public:
+
+	static void Initialize();
+
+	static void Shutdown();
+
+	static void ReloadTextures();
+
+	static const ISlateStyle& Get();
+
+	static FName GetStyleSetName();
+
+	static const FLinearColor& GetColor(FName PropertyName, const ANSICHAR* Specifier = nullptr);
+
+	static const FSlateBrush* GetBrush(FName PropertyName, const ANSICHAR* Specifier = nullptr);
+
+	template< typename WidgetStyleType >
+	static const WidgetStyleType& GetWidgetStyle(FName PropertyName, const ANSICHAR* Specifier = nullptr)
+	{
+		return StyleInstance->GetWidgetStyle<WidgetStyleType>(PropertyName, Specifier);
+	}
+
+private:
+	static TSharedRef<FSlateStyleSet> Create();
+
+private:
+	static TSharedPtr<FSlateStyleSet> StyleInstance;
+};
diff --git a/Engine/Plugins/VirtualProduction/LevelSnapshots/Source/LevelSnapshotsEditor/Private/LevelSnapshotsUserSettings.h b/Engine/Plugins/VirtualProduction/LevelSnapshots/Source/LevelSnapshotsEditor/Private/LevelSnapshotsUserSettings.h
new file mode 100644
--- /dev/null
+++ b/Engine/Plugins/VirtualProduction/LevelSnapshots/Source/LevelSnapshotsEditor/Private/LevelSnapshotsUserSettings.h
@@ -0,0 +1,18 @@
+// Copyright Epic Games, Inc. All Rights Reserved.
+
+#pragma once
+
+#include "UObject/Object.h"
+#include "LevelSnapshotsUserSettings.generated.h"
+
+class ULevelSnapshotsEditorData;
+
+UCLASS(config = EditorPerProjectUserSettings, MinimalAPI)
+class ULevelSnapshotsUserSettings : public UObject
+{
+public:
+	GENERATED_BODY()
+
+	UPROPERTY(config)
+	TSoftObjectPtr<ULevelSnapshotsEditorData> LastEditorData;
+};
\ No newline at end of file
diff --git a/Engine/Plugins/VirtualProduction/LevelSnapshots/Source/LevelSnapshotsEditor/Private/Misc/LevelSnapshotsEditorContext.cpp b/Engine/Plugins/VirtualProduction/LevelSnapshots/Source/LevelSnapshotsEditor/Private/Misc/LevelSnapshotsEditorContext.cpp
new file mode 100644
--- /dev/null
+++ b/Engine/Plugins/VirtualProduction/LevelSnapshots/Source/LevelSnapshotsEditor/Private/Misc/LevelSnapshotsEditorContext.cpp
@@ -0,0 +1,91 @@
+// Copyright Epic Games, Inc. All Rights Reserved.
+
+#include "Misc/LevelSnapshotsEditorContext.h"
+
+#include "Engine/World.h"
+#include "Engine/Engine.h"
+#include "Editor.h"
+
+#define LOCTEXT_NAMESPACE "LevelSnapshotsEditor"
+
+FLevelSnapshotsEditorContext::FLevelSnapshotsEditorContext()
+{
+	FEditorDelegates::MapChange.AddRaw(this, &FLevelSnapshotsEditorContext::OnMapChange);
+	FEditorDelegates::PreBeginPIE.AddRaw(this, &FLevelSnapshotsEditorContext::OnPieEvent);
+	FEditorDelegates::BeginPIE.AddRaw(this, &FLevelSnapshotsEditorContext::OnPieEvent);
+	FEditorDelegates::PostPIEStarted.AddRaw(this, &FLevelSnapshotsEditorContext::OnPieEvent);
+	FEditorDelegates::PrePIEEnded.AddRaw(this, &FLevelSnapshotsEditorContext::OnPieEvent);
+	FEditorDelegates::EndPIE.AddRaw(this, &FLevelSnapshotsEditorContext::OnPieEvent);
+}
+
+FLevelSnapshotsEditorContext::~FLevelSnapshotsEditorContext()
+{
+	FEditorDelegates::MapChange.RemoveAll(this);
+	FEditorDelegates::PreBeginPIE.RemoveAll(this);
+	FEditorDelegates::BeginPIE.RemoveAll(this);
+	FEditorDelegates::PostPIEStarted.RemoveAll(this);
+	FEditorDelegates::PrePIEEnded.RemoveAll(this);
+	FEditorDelegates::EndPIE.RemoveAll(this);
+}
+
+void FLevelSnapshotsEditorContext::OnPieEvent(bool)
+{
+	WeakCurrentContext = nullptr;
+}
+
+void FLevelSnapshotsEditorContext::OnMapChange(uint32)
+{
+	WeakCurrentContext = nullptr;
+}
+
+UWorld* FLevelSnapshotsEditorContext::Get() const
+{
+	UWorld* Context = WeakCurrentContext.Get();
+	if (Context)
+	{
+		return Context;
+	}
+
+	Context = ComputeContext();
+	check(Context);
+	WeakCurrentContext = Context;
+	return Context;
+}
+
+UObject* FLevelSnapshotsEditorContext::GetAsObject() const
+{
+	return Get();
+}
+
+void FLevelSnapshotsEditorContext::OverrideWith(UWorld* InNewContext)
+{
+	WeakCurrentContext = InNewContext;
+}
+
+UWorld* FLevelSnapshotsEditorContext::ComputeContext()
+{
+	UWorld* EditorWorld = nullptr;
+
+	// Return PIE worlds if there are any
+	for (const FWorldContext& Context : GEngine->GetWorldContexts())
+	{
+		if (Context.WorldType == EWorldType::PIE)
+		{
+			UWorld* ThisWorld = Context.World();
+			if (GEditor && GEditor->bIsSimulatingInEditor)
+			{
+				return ThisWorld;
+			}
+		}
+		else if (Context.WorldType == EWorldType::Editor)
+		{
+			EditorWorld = Context.World();
+		}
+	}
+
+	check(EditorWorld);
+
+	return EditorWorld;
+}
+
+#undef LOCTEXT_NAMESPACE
diff --git a/Engine/Plugins/VirtualProduction/LevelSnapshots/Source/LevelSnapshotsEditor/Private/Misc/LevelSnapshotsEditorContext.h b/Engine/Plugins/VirtualProduction/LevelSnapshots/Source/LevelSnapshotsEditor/Private/Misc/LevelSnapshotsEditorContext.h
new file mode 100644
--- /dev/null
+++ b/Engine/Plugins/VirtualProduction/LevelSnapshots/Source/LevelSnapshotsEditor/Private/Misc/LevelSnapshotsEditorContext.h
@@ -0,0 +1,52 @@
+// Copyright Epic Games, Inc. All Rights Reserved.
+
+#pragma once
+
+#include "ILevelSnapshotsEditorView.h"
+
+#include "Templates/SharedPointer.h"
+#include "Delegates/Delegate.h"
+#include "UObject/WeakObjectPtr.h"
+#include "UObject/WeakObjectPtrTemplates.h"
+
+class UWorld;
+
+class FLevelSnapshotsEditorContext : public ILevelSnapshotsEditorContext
+{
+public:
+
+	FLevelSnapshotsEditorContext();
+	~FLevelSnapshotsEditorContext();
+
+	/**
+	 * Resolve the current world context pointer. Can never be nullptr.
+	 */
+	virtual UWorld* Get() const override;
+
+	/**
+	 * Resolve the current world context pointer as a base object. Can never be nullptr.
+	 */
+	UObject* GetAsObject() const;
+
+	/**
+	 * Compute the new playback context based on the user's current auto-bind settings.
+	 * Will use the first encountered PIE or Simulate world if possible, else the Editor world as a fallback
+	 */
+	static UWorld* ComputeContext();
+
+	/**
+	 * Specify a new world to use as the context. Persists until the next PIE or map change event.
+	 * May be null, in which case the context will be recomputed automatically
+	 */
+	void OverrideWith(UWorld* InNewContext);
+
+private:
+
+	void OnPieEvent(bool);
+	void OnMapChange(uint32);
+
+private:
+
+	/** Mutable cached context pointer */
+	mutable TWeakObjectPtr<UWorld> WeakCurrentContext;
+};
\ No newline at end of file
diff --git a/Engine/Plugins/VirtualProduction/LevelSnapshots/Source/LevelSnapshotsEditor/Private/TempInterfaces/ILevelSnapshotsEditorInput.h b/Engine/Plugins/VirtualProduction/LevelSnapshots/Source/LevelSnapshotsEditor/Private/TempInterfaces/ILevelSnapshotsEditorInput.h
new file mode 100644
--- /dev/null
+++ b/Engine/Plugins/VirtualProduction/LevelSnapshots/Source/LevelSnapshotsEditor/Private/TempInterfaces/ILevelSnapshotsEditorInput.h
@@ -0,0 +1,10 @@
+// Copyright Epic Games, Inc. All Rights Reserved.
+
+#pragma once
+
+#include "ILevelSnapshotsEditorView.h"
+
+class ILevelSnapshotsEditorInput : public ILevelSnapshotsEditorView
+{
+public:
+};
diff --git a/Engine/Plugins/VirtualProduction/LevelSnapshots/Source/LevelSnapshotsEditor/Private/TempInterfaces/ILevelSnapshotsEditorResults.h b/Engine/Plugins/VirtualProduction/LevelSnapshots/Source/LevelSnapshotsEditor/Private/TempInterfaces/ILevelSnapshotsEditorResults.h
new file mode 100644
--- /dev/null
+++ b/Engine/Plugins/VirtualProduction/LevelSnapshots/Source/LevelSnapshotsEditor/Private/TempInterfaces/ILevelSnapshotsEditorResults.h
@@ -0,0 +1,10 @@
+// Copyright Epic Games, Inc. All Rights Reserved.
+
+#pragma once
+
+#include "ILevelSnapshotsEditorView.h"
+
+class ILevelSnapshotsEditorResults : public ILevelSnapshotsEditorView
+{
+public:
+};
diff --git a/Engine/Plugins/VirtualProduction/LevelSnapshots/Source/LevelSnapshotsEditor/Private/Views/Input/LevelSnapshotsEditorInput.cpp b/Engine/Plugins/VirtualProduction/LevelSnapshots/Source/LevelSnapshotsEditor/Private/Views/Input/LevelSnapshotsEditorInput.cpp
new file mode 100644
--- /dev/null
+++ b/Engine/Plugins/VirtualProduction/LevelSnapshots/Source/LevelSnapshotsEditor/Private/Views/Input/LevelSnapshotsEditorInput.cpp
@@ -0,0 +1,20 @@
+// Copyright Epic Games, Inc. All Rights Reserved.
+
+#include "Views/Input/LevelSnapshotsEditorInput.h"
+
+#include "Views/Input/SLevelSnapshotsEditorInput.h"
+
+FLevelSnapshotsEditorInput::FLevelSnapshotsEditorInput(const TSharedRef<FLevelSnapshotsEditorViewBuilder>& InBuilder)
+	: BuilderPtr(InBuilder)
+{
+}
+
+TSharedRef<SWidget> FLevelSnapshotsEditorInput::GetOrCreateWidget()
+{
+	if (!EditorInputWidget.IsValid())
+	{
+		SAssignNew(EditorInputWidget, SLevelSnapshotsEditorInput, SharedThis(this), BuilderPtr.Pin().ToSharedRef());
+	}
+
+	return EditorInputWidget.ToSharedRef();
+}
diff --git a/Engine/Plugins/VirtualProduction/LevelSnapshots/Source/LevelSnapshotsEditor/Private/Views/Input/LevelSnapshotsEditorInput.h b/Engine/Plugins/VirtualProduction/LevelSnapshots/Source/LevelSnapshotsEditor/Private/Views/Input/LevelSnapshotsEditorInput.h
new file mode 100644
--- /dev/null
+++ b/Engine/Plugins/VirtualProduction/LevelSnapshots/Source/LevelSnapshotsEditor/Private/Views/Input/LevelSnapshotsEditorInput.h
@@ -0,0 +1,23 @@
+// Copyright Epic Games, Inc. All Rights Reserved.
+
+#pragma once
+
+#include "ILevelSnapshotsEditorInput.h"
+
+class SLevelSnapshotsEditorInput;
+
+class FLevelSnapshotsEditorInput : public ILevelSnapshotsEditorInput
+{
+public:
+	FLevelSnapshotsEditorInput(const TSharedRef<FLevelSnapshotsEditorViewBuilder>& InBuilder);
+
+	//~ Begin ILevelSnapshotsEditorView Interface
+	virtual TSharedRef<SWidget> GetOrCreateWidget() override;
+	virtual TSharedRef<FLevelSnapshotsEditorViewBuilder> GetBuilder() const override { return BuilderPtr.Pin().ToSharedRef(); }
+	//~ End ILevelSnapshotsEditorView Interface
+
+private:
+	TSharedPtr<SLevelSnapshotsEditorInput> EditorInputWidget;
+
+	TWeakPtr<FLevelSnapshotsEditorViewBuilder> BuilderPtr;
+};
diff --git a/Engine/Plugins/VirtualProduction/LevelSnapshots/Source/LevelSnapshotsEditor/Private/Widgets/SLevelSnapshotsEditorBrowser.h b/Engine/Plugins/VirtualProduction/LevelSnapshots/Source/LevelSnapshotsEditor/Private/Widgets/SLevelSnapshotsEditorBrowser.h
new file mode 100644
--- /dev/null
+++ b/Engine/Plugins/VirtualProduction/LevelSnapshots/Source/LevelSnapshotsEditor/Private/Widgets/SLevelSnapshotsEditorBrowser.h
@@ -0,0 +1,32 @@
+// Copyright Epic Games, Inc. All Rights Reserved.
+
+#pragma once
+
+#include "Widgets/DeclarativeSyntaxSupport.h"
+#include "Widgets/SCompoundWidget.h"
+
+struct FAssetData;
+class UWorld;
+struct FLevelSnapshotsEditorViewBuilder;
+
+class SLevelSnapshotsEditorBrowser : public SCompoundWidget
+{
+public:
+	SLATE_BEGIN_ARGS(SLevelSnapshotsEditorBrowser)
+	{}
+
+	/** Attribute for retrieving the current context */
+	SLATE_ATTRIBUTE(FSoftObjectPath, OwningWorldPath)
+
+	SLATE_END_ARGS()
+
+	void Construct(const FArguments& InArgs, const TSharedRef<FLevelSnapshotsEditorViewBuilder>& InBuilder);
+
+private:
+	void OnAssetSelected(const FAssetData& InAssetData);
+	bool OnShouldFilterAsset(const FAssetData& InAssetData) const;
+
+	TAttribute<FSoftObjectPath> OwningWorldPathAttribute;
+
+	TWeakPtr<FLevelSnapshotsEditorViewBuilder> BuilderPtr;
+};
