Slate: Deprecate SSpliter's slot public variable.

diff --git a/Engine/Plugins/Enterprise/VariantManager/Source/VariantManager/Private/SVariantManager.cpp b/Engine/Plugins/Enterprise/VariantManager/Source/VariantManager/Private/SVariantManager.cpp
--- a/Engine/Plugins/Enterprise/VariantManager/Source/VariantManager/Private/SVariantManager.cpp
+++ b/Engine/Plugins/Enterprise/VariantManager/Source/VariantManager/Private/SVariantManager.cpp
@@ -436,7 +436,7 @@ SVariantManager::~SVariantManager()
 		if (MainSlots->Num() > 0)
 		{
 			FSplitterValues Values;
-			Values.VariantColumn = MainSplitter->SlotAt(0).SizeValue.Get();
+			Values.VariantColumn = MainSplitter->SlotAt(0).GetSizeValue();
 			Values.ActorColumn = OnGetPropertiesActorColumnWidth();;
 			Values.PropertyNameColumn = OnGetPropertiesNameColumnWidth();
 			Values.DependenciesVariantSetsColumn = OnGetDependenciesVariantSetColumnWidth();
diff --git a/Engine/Plugins/VirtualProduction/LevelSnapshots/Source/LevelSnapshotsEditor/Private/Views/Results/SLevelSnapshotsEditorResults.cpp b/Engine/Plugins/VirtualProduction/LevelSnapshots/Source/LevelSnapshotsEditor/Private/Views/Results/SLevelSnapshotsEditorResults.cpp
--- a/Engine/Plugins/VirtualProduction/LevelSnapshots/Source/LevelSnapshotsEditor/Private/Views/Results/SLevelSnapshotsEditorResults.cpp
+++ b/Engine/Plugins/VirtualProduction/LevelSnapshots/Source/LevelSnapshotsEditor/Private/Views/Results/SLevelSnapshotsEditorResults.cpp
@@ -2649,6 +2649,8 @@ void SLevelSnapshotsEditorResultsRow::Construct(const FArguments& InArgs, const
 				.HitDetectionSplitterHandleSize(5.0f)
 
 				+ SSplitter::Slot()
+				.OnSlotResized(this, &SLevelSnapshotsEditorResultsRow::SetNameColumnSize)
+				.Value(this, &SLevelSnapshotsEditorResultsRow::GetSplitterSlotSize, 0)
 				[
 					SNew(SHorizontalBox)
 
@@ -2678,12 +2680,6 @@ void SLevelSnapshotsEditorResultsRow::Construct(const FArguments& InArgs, const
 
 	// Create value widgets
 
-	// Splitter Slot 0
-	SplitterPtr->SlotAt(0).OnSlotResized_Handler.BindSP(this, &SLevelSnapshotsEditorResultsRow::SetNameColumnSize);
-
-	const auto SlotDelegate0 = TAttribute<float>::FGetter::CreateSP(this, &SLevelSnapshotsEditorResultsRow::GetSplitterSlotSize, 0);;
-	SplitterPtr->SlotAt(0).SizeValue.Bind(SlotDelegate0);
-		
 	// Splitter Slot 1
 	TSharedPtr<SWidget> SnapshotChildWidget;
 	
@@ -2760,13 +2756,13 @@ void SLevelSnapshotsEditorResultsRow::Construct(const FArguments& InArgs, const
 	}
 
 	SplitterPtr->AddSlot()
+	.OnSlotResized(this, &SLevelSnapshotsEditorResultsRow::SetSnapshotColumnSize)
+	.Value(this, &SLevelSnapshotsEditorResultsRow::GetSplitterSlotSize, 1)
 	[
 		FinalSnapshotWidget.ToSharedRef()
-	].OnSlotResized(this, &SLevelSnapshotsEditorResultsRow::SetSnapshotColumnSize);
-
-	const auto SlotDelegate1 = TAttribute<float>::FGetter::CreateSP(this, &SLevelSnapshotsEditorResultsRow::GetSplitterSlotSize, 1);;
-	SplitterPtr->SlotAt(1).SizeValue.Bind(SlotDelegate1);
+	];
 
+	
 	// Splitter Slot 2
 	TSharedPtr<SWidget> WorldChildWidget;
 
@@ -2843,12 +2839,11 @@ void SLevelSnapshotsEditorResultsRow::Construct(const FArguments& InArgs, const
 	}
 
 	SplitterPtr->AddSlot()
+	.OnSlotResized(this, &SLevelSnapshotsEditorResultsRow::SetWorldObjectColumnSize)
+	.Value(this, &SLevelSnapshotsEditorResultsRow::GetSplitterSlotSize, 2)
 	[
 		FinalWorldWidget.ToSharedRef()
-	].OnSlotResized(this, &SLevelSnapshotsEditorResultsRow::SetWorldObjectColumnSize);
-
-	const auto SlotDelegate2 = TAttribute<float>::FGetter::CreateSP(this, &SLevelSnapshotsEditorResultsRow::GetSplitterSlotSize, 2);;
-	SplitterPtr->SlotAt(2).SizeValue.Bind(SlotDelegate2);
+	];
 }
 
 SLevelSnapshotsEditorResultsRow::~SLevelSnapshotsEditorResultsRow()
@@ -2860,7 +2855,7 @@ SLevelSnapshotsEditorResultsRow::~SLevelSnapshotsEditorResultsRow()
 	{
 		for (int32 SplitterSlotCount = 0; SplitterSlotCount < SplitterPtr->GetChildren()->Num(); SplitterSlotCount++)
 		{
-			SplitterPtr->SlotAt(SplitterSlotCount).OnSlotResized_Handler.Unbind();
+			SplitterPtr->SlotAt(SplitterSlotCount).OnSlotResized().Unbind();
 		}
 	}
 
diff --git a/Engine/Source/Developer/LogVisualizer/Private/SVisualLoggerView.cpp b/Engine/Source/Developer/LogVisualizer/Private/SVisualLoggerView.cpp
--- a/Engine/Source/Developer/LogVisualizer/Private/SVisualLoggerView.cpp
+++ b/Engine/Source/Developer/LogVisualizer/Private/SVisualLoggerView.cpp
@@ -264,7 +264,7 @@ void SVisualLoggerView::Construct(const FArguments& InArgs, const TSharedRef<FUI
 			SSplitter::FSlot const& LeftSplitterSlot = SearchSplitter->SlotAt(0);
 			SSplitter::FSlot const& RightSplitterSlot = SearchSplitter->SlotAt(1);
 
-			return LeftSplitterSlot.SizeValue.Get() / RightSplitterSlot.SizeValue.Get();
+			return LeftSplitterSlot.GetSizeValue() / RightSplitterSlot.GetSizeValue();
 	});
 
 	OnSearchSplitterResized();
@@ -294,7 +294,7 @@ void SVisualLoggerView::OnSearchSplitterResized()
 	SSplitter::FSlot const& LeftSplitterSlot = SearchSplitter->SlotAt(0);
 	SSplitter::FSlot const& RightSplitterSlot = SearchSplitter->SlotAt(1);
 
-	const float NewAnimationOutlinerFillPercentage = LeftSplitterSlot.SizeValue.Get() / RightSplitterSlot.SizeValue.Get();
+	const float NewAnimationOutlinerFillPercentage = LeftSplitterSlot.GetSizeValue() / RightSplitterSlot.GetSizeValue();
 	SetAnimationOutlinerFillPercentage(NewAnimationOutlinerFillPercentage);
 
 	FLogVisualizer::Get().SetAnimationOutlinerFillPercentage(NewAnimationOutlinerFillPercentage);
diff --git a/Engine/Source/Developer/LogVisualizer/Private/SVisualLoggerView.h b/Engine/Source/Developer/LogVisualizer/Private/SVisualLoggerView.h
--- a/Engine/Source/Developer/LogVisualizer/Private/SVisualLoggerView.h
+++ b/Engine/Source/Developer/LogVisualizer/Private/SVisualLoggerView.h
@@ -35,7 +35,7 @@ public:
 		SSplitter::FSlot const& LeftSplitterSlot = SearchSplitter->SlotAt(0);
 		SSplitter::FSlot const& RightSplitterSlot = SearchSplitter->SlotAt(1);
 
-		return LeftSplitterSlot.SizeValue.Get()/ RightSplitterSlot.SizeValue.Get();
+		return LeftSplitterSlot.GetSizeValue()/ RightSplitterSlot.GetSizeValue();
 	}
 	void SetAnimationOutlinerFillPercentage(float FillPercentage);
 
diff --git a/Engine/Source/Editor/ContentBrowser/Private/SContentBrowser.cpp b/Engine/Source/Editor/ContentBrowser/Private/SContentBrowser.cpp
--- a/Engine/Source/Editor/ContentBrowser/Private/SContentBrowser.cpp
+++ b/Engine/Source/Editor/ContentBrowser/Private/SContentBrowser.cpp
@@ -1355,13 +1355,13 @@ void SContentBrowser::SaveSettings() const
 
 	for(int32 SlotIndex = 0; SlotIndex < PathAssetSplitterPtr->GetChildren()->Num(); SlotIndex++)
 	{
-		float SplitterSize = PathAssetSplitterPtr->SlotAt(SlotIndex).SizeValue.Get();
+		float SplitterSize = PathAssetSplitterPtr->SlotAt(SlotIndex).GetSizeValue();
 		GConfig->SetFloat(*SettingsIniSection, *(SettingsString + FString::Printf(TEXT(".VerticalSplitter.SlotSize%d"), SlotIndex)), SplitterSize, GEditorPerProjectIni);
 	}
 
 	for (int32 SlotIndex = 0; SlotIndex < PathFavoriteSplitterPtr->GetChildren()->Num(); SlotIndex++)
 	{
-		float SplitterSize = PathFavoriteSplitterPtr->SlotAt(SlotIndex).SizeValue.Get();
+		float SplitterSize = PathFavoriteSplitterPtr->SlotAt(SlotIndex).GetSizeValue();
 		GConfig->SetFloat(*SettingsIniSection, *(SettingsString + FString::Printf(TEXT(".FavoriteSplitter.SlotSize%d"), SlotIndex)), SplitterSize, GEditorPerProjectIni);
 	}
 
@@ -1501,16 +1501,16 @@ void SContentBrowser::LoadSettings(const FName& InInstanceName)
 
 	for(int32 SlotIndex = 0; SlotIndex < PathAssetSplitterPtr->GetChildren()->Num(); SlotIndex++)
 	{
-		float SplitterSize = PathAssetSplitterPtr->SlotAt(SlotIndex).SizeValue.Get();
+		float SplitterSize = PathAssetSplitterPtr->SlotAt(SlotIndex).GetSizeValue();
 		GConfig->GetFloat(*SettingsIniSection, *(SettingsString + FString::Printf(TEXT(".VerticalSplitter.SlotSize%d"), SlotIndex)), SplitterSize, GEditorPerProjectIni);
-		PathAssetSplitterPtr->SlotAt(SlotIndex).SizeValue = SplitterSize;
+		PathAssetSplitterPtr->SlotAt(SlotIndex).SetSizeValue(SplitterSize);
 	}
 
 	for (int32 SlotIndex = 0; SlotIndex < PathFavoriteSplitterPtr->GetChildren()->Num(); SlotIndex++)
 	{
-		float SplitterSize = PathFavoriteSplitterPtr->SlotAt(SlotIndex).SizeValue.Get();
+		float SplitterSize = PathFavoriteSplitterPtr->SlotAt(SlotIndex).GetSizeValue();
 		GConfig->GetFloat(*SettingsIniSection, *(SettingsString + FString::Printf(TEXT(".FavoriteSplitter.SlotSize%d"), SlotIndex)), SplitterSize, GEditorPerProjectIni);
-		PathFavoriteSplitterPtr->SlotAt(SlotIndex).SizeValue = SplitterSize;
+		PathFavoriteSplitterPtr->SlotAt(SlotIndex).SetSizeValue(SplitterSize);
 	}
 
 	// Save all our data using the settings string as a key in the user settings ini
diff --git a/Engine/Source/Editor/Kismet/Private/SBlueprintPalette.cpp b/Engine/Source/Editor/Kismet/Private/SBlueprintPalette.cpp
--- a/Engine/Source/Editor/Kismet/Private/SBlueprintPalette.cpp
+++ b/Engine/Source/Editor/Kismet/Private/SBlueprintPalette.cpp
@@ -2012,11 +2012,11 @@ void SBlueprintPalette::OnSplitterResized() const
 
 		if (SplitterSlot.GetWidget() == FavoritesWrapper)
 		{
-			GConfig->SetFloat(*BlueprintPalette::ConfigSection, *BlueprintPalette::FavoritesHeightConfigKey, SplitterSlot.SizeValue.Get(), GEditorPerProjectIni);
+			GConfig->SetFloat(*BlueprintPalette::ConfigSection, *BlueprintPalette::FavoritesHeightConfigKey, SplitterSlot.GetSizeValue(), GEditorPerProjectIni);
 		}
 		else if (SplitterSlot.GetWidget() == LibraryWrapper)
 		{
-			GConfig->SetFloat(*BlueprintPalette::ConfigSection, *BlueprintPalette::LibraryHeightConfigKey, SplitterSlot.SizeValue.Get(), GEditorPerProjectIni);
+			GConfig->SetFloat(*BlueprintPalette::ConfigSection, *BlueprintPalette::LibraryHeightConfigKey, SplitterSlot.GetSizeValue(), GEditorPerProjectIni);
 		}
 
 	}
diff --git a/Engine/Source/Editor/UnrealEd/Private/EditorViewportLayoutFourPanes.cpp b/Engine/Source/Editor/UnrealEd/Private/EditorViewportLayoutFourPanes.cpp
--- a/Engine/Source/Editor/UnrealEd/Private/EditorViewportLayoutFourPanes.cpp
+++ b/Engine/Source/Editor/UnrealEd/Private/EditorViewportLayoutFourPanes.cpp
@@ -97,10 +97,10 @@ void FEditorViewportLayoutFourPanes::SaveLayoutString(const FString& SpecificLay
 	const FString& IniSection = FLayoutSaveRestore::GetAdditionalLayoutConfigIni();
 
 	check(PrimarySplitterWidget->GetChildren()->Num() == 2);
-	float PrimaryPercentage = PrimarySplitterWidget->SlotAt(0).SizeValue.Get();
+	float PrimaryPercentage = PrimarySplitterWidget->SlotAt(0).GetSizeValue();
 	check(SecondarySplitterWidget->GetChildren()->Num() == 3);
-	float SecondaryPercentage0 = SecondarySplitterWidget->SlotAt(0).SizeValue.Get();
-	float SecondaryPercentage1 = SecondarySplitterWidget->SlotAt(1).SizeValue.Get();
+	float SecondaryPercentage0 = SecondarySplitterWidget->SlotAt(0).GetSizeValue();
+	float SecondaryPercentage1 = SecondarySplitterWidget->SlotAt(1).GetSizeValue();
 
 	GConfig->SetString(*IniSection, *(SpecificLayoutString + TEXT(".Percentage0")), *TTypeToString<float>::ToString(PrimaryPercentage), GEditorPerProjectIni);
 	GConfig->SetString(*IniSection, *(SpecificLayoutString + TEXT(".Percentage1")), *TTypeToString<float>::ToString(SecondaryPercentage0), GEditorPerProjectIni);
diff --git a/Engine/Source/Editor/UnrealEd/Private/EditorViewportLayoutThreePanes.cpp b/Engine/Source/Editor/UnrealEd/Private/EditorViewportLayoutThreePanes.cpp
--- a/Engine/Source/Editor/UnrealEd/Private/EditorViewportLayoutThreePanes.cpp
+++ b/Engine/Source/Editor/UnrealEd/Private/EditorViewportLayoutThreePanes.cpp
@@ -89,9 +89,9 @@ void FEditorViewportLayoutThreePanes::SaveLayoutString(const FString& SpecificLa
 	const FString& IniSection = FLayoutSaveRestore::GetAdditionalLayoutConfigIni();
 
 	check(PrimarySplitterWidget->GetChildren()->Num() == 2);
-	float PrimaryPercentage = PrimarySplitterWidget->SlotAt(0).SizeValue.Get();
+	float PrimaryPercentage = PrimarySplitterWidget->SlotAt(0).GetSizeValue();
 	check(SecondarySplitterWidget->GetChildren()->Num() == 2);
-	float SecondaryPercentage = SecondarySplitterWidget->SlotAt(0).SizeValue.Get();
+	float SecondaryPercentage = SecondarySplitterWidget->SlotAt(0).GetSizeValue();
 
 	GConfig->SetString(*IniSection, *(SpecificLayoutString + TEXT(".Percentage0")), *TTypeToString<float>::ToString(PrimaryPercentage), GEditorPerProjectIni);
 	GConfig->SetString(*IniSection, *(SpecificLayoutString + TEXT(".Percentage1")), *TTypeToString<float>::ToString(SecondaryPercentage), GEditorPerProjectIni);
diff --git a/Engine/Source/Editor/UnrealEd/Private/EditorViewportLayoutTwoPanes.cpp b/Engine/Source/Editor/UnrealEd/Private/EditorViewportLayoutTwoPanes.cpp
--- a/Engine/Source/Editor/UnrealEd/Private/EditorViewportLayoutTwoPanes.cpp
+++ b/Engine/Source/Editor/UnrealEd/Private/EditorViewportLayoutTwoPanes.cpp
@@ -100,7 +100,7 @@ void TEditorViewportLayoutTwoPanes<TOrientation>::SaveLayoutString(const FString
 	const FString& IniSection = FLayoutSaveRestore::GetAdditionalLayoutConfigIni();
 
 	check(SplitterWidget->GetChildren()->Num() == 2);
-	float Percentage = SplitterWidget->SlotAt(0).SizeValue.Get();
+	float Percentage = SplitterWidget->SlotAt(0).GetSizeValue();
 
 	GConfig->SetString(*IniSection, *(SpecificLayoutString + TEXT(".Percentage")), *TTypeToString<float>::ToString(Percentage), GEditorPerProjectIni);
 }
diff --git a/Engine/Source/Runtime/Slate/Private/Widgets/Layout/SSplitter.cpp b/Engine/Source/Runtime/Slate/Private/Widgets/Layout/SSplitter.cpp
--- a/Engine/Source/Runtime/Slate/Private/Widgets/Layout/SSplitter.cpp
+++ b/Engine/Source/Runtime/Slate/Private/Widgets/Layout/SSplitter.cpp
@@ -4,6 +4,7 @@
 #include "Layout/ArrangedChildren.h"
 #include "Rendering/DrawElements.h"
 
+PRAGMA_DISABLE_DEPRECATION_WARNINGS
 void SSplitter::FSlot::Construct(const FChildren& SlotOwner, FSlotArguments&& InArgs)
 {
 	TSlotBase<FSlot>::Construct(SlotOwner, MoveTemp(InArgs));
@@ -833,6 +834,7 @@ int32 SSplitter::GetHandleBeingResizedFromMousePosition( float InPhysicalSplitte
 
 	return INDEX_NONE;
 }
+PRAGMA_ENABLE_DEPRECATION_WARNINGS
 
 /********************************************************************
 * SSplitter2x2														*
diff --git a/Engine/Source/Runtime/Slate/Public/Widgets/Layout/SSplitter.h b/Engine/Source/Runtime/Slate/Public/Widgets/Layout/SSplitter.h
--- a/Engine/Source/Runtime/Slate/Public/Widgets/Layout/SSplitter.h
+++ b/Engine/Source/Runtime/Slate/Public/Widgets/Layout/SSplitter.h
@@ -60,6 +60,7 @@ public:
 	DECLARE_DELEGATE_RetVal_OneParam(FVector2D, FOnGetMaxSlotSize, int32);
 
 public:
+	PRAGMA_DISABLE_DEPRECATION_WARNINGS
 	class SLATE_API FSlot : public TSlotBase<FSlot>
 	{
 	public:		
@@ -146,12 +147,18 @@ public:
 		bool CanBeResized() const;
 
 	public:
+		UE_DEPRECATED(5.0, "Direct access to SizingRule is now deprecated. Use the getter.")
 		TAttribute<ESizeRule> SizingRule;
+		UE_DEPRECATED(5.0, "Direct access to SizeValue is now deprecated. Use the getter.")
 		TAttribute<float> SizeValue;
+		UE_DEPRECATED(5.0, "Direct access to MinSizeValue is now deprecated. Use the getter.")
 		TOptional<float> MinSizeValue;
+		UE_DEPRECATED(5.0, "Direct access to OnSlotResized_Handler is now deprecated. Use the getter.")
 		FOnSlotResized OnSlotResized_Handler;
+		UE_DEPRECATED(5.0, "Direct access to bIsResizable is now deprecated. Use the getter.")
 		TOptional<bool> bIsResizable;
 	};
+	PRAGMA_ENABLE_DEPRECATION_WARNINGS
 
 	/** @return Add a new FSlot() */
 	static FSlot::FSlotArguments Slot();
