Fixes for no unity build errors
#rnx

diff --git a/Engine/Plugins/Media/PixelStreaming/Source/PixelStreaming/Private/LatencyTester.h b/Engine/Plugins/Media/PixelStreaming/Source/PixelStreaming/Private/LatencyTester.h
--- a/Engine/Plugins/Media/PixelStreaming/Source/PixelStreaming/Private/LatencyTester.h
+++ b/Engine/Plugins/Media/PixelStreaming/Source/PixelStreaming/Private/LatencyTester.h
@@ -3,6 +3,7 @@
 #pragma once
 
 #include "Containers/UnrealString.h"
+#include <atomic>
 
 // Performs a latency test of one frame passing through Pixel Streaming
 // from capture, to encode, to transmit
diff --git a/Engine/Plugins/Media/PixelStreaming/Source/PixelStreaming/Public/PixelStreamingSettings.h b/Engine/Plugins/Media/PixelStreaming/Source/PixelStreaming/Public/PixelStreamingSettings.h
--- a/Engine/Plugins/Media/PixelStreaming/Source/PixelStreaming/Public/PixelStreamingSettings.h
+++ b/Engine/Plugins/Media/PixelStreaming/Source/PixelStreaming/Public/PixelStreamingSettings.h
@@ -4,7 +4,11 @@
 
 #include "CoreMinimal.h"
 #include "Engine/DeveloperSettings.h"
+#include "HAL/IConsoleManager.h"
 #include "InputCoreTypes.h"
+#include "Misc/CommandLine.h"
+#include "VideoEncoder.h"
+#include "WebRTCIncludes.h"
 #include "PixelStreamingSettings.generated.h"
 
 // Console variables (CVars)
diff --git a/Engine/Plugins/Runtime/AudioModulation/Source/AudioModulation/Public/SoundModulationGenerator.h b/Engine/Plugins/Runtime/AudioModulation/Source/AudioModulation/Public/SoundModulationGenerator.h
--- a/Engine/Plugins/Runtime/AudioModulation/Source/AudioModulation/Public/SoundModulationGenerator.h
+++ b/Engine/Plugins/Runtime/AudioModulation/Source/AudioModulation/Public/SoundModulationGenerator.h
@@ -2,6 +2,7 @@
 #pragma once
 
 #include "AudioDeviceManager.h"
+#include "Containers/Queue.h"
 #include "CoreMinimal.h"
 #include "DSP/LFO.h"
 #include "IAudioModulation.h"
diff --git a/Engine/Plugins/Runtime/NetworkPrediction/Source/NetworkPrediction/Private/NetworkPredictionAsync.cpp b/Engine/Plugins/Runtime/NetworkPrediction/Source/NetworkPrediction/Private/NetworkPredictionAsync.cpp
--- a/Engine/Plugins/Runtime/NetworkPrediction/Source/NetworkPrediction/Private/NetworkPredictionAsync.cpp
+++ b/Engine/Plugins/Runtime/NetworkPrediction/Source/NetworkPrediction/Private/NetworkPredictionAsync.cpp
@@ -1,6 +1,8 @@
 // Copyright Epic Games, Inc. All Rights Reserved.
 
 #include "NetworkPredictionAsync.h"
+#include "Components/PrimitiveComponent.h"
+#include "DrawDebugHelpers.h"
 #include "Net/UnrealNetwork.h"
 #include "PhysicsProxy/SingleParticlePhysicsProxy.h"
 
diff --git a/Engine/Plugins/VirtualProduction/LevelSnapshots/Source/LevelSnapshots/Private/Archive/ClassDefaults/BaseClassDefaultArchive.cpp b/Engine/Plugins/VirtualProduction/LevelSnapshots/Source/LevelSnapshots/Private/Archive/ClassDefaults/BaseClassDefaultArchive.cpp
--- a/Engine/Plugins/VirtualProduction/LevelSnapshots/Source/LevelSnapshots/Private/Archive/ClassDefaults/BaseClassDefaultArchive.cpp
+++ b/Engine/Plugins/VirtualProduction/LevelSnapshots/Source/LevelSnapshots/Private/Archive/ClassDefaults/BaseClassDefaultArchive.cpp
@@ -1,7 +1,8 @@
 ﻿// Copyright Epic Games, Inc. All Rights Reserved.
 
 #include "BaseClassDefaultArchive.h"
-
+#include "GameFramework/Actor.h"
+#include "Serialization/ArchiveSerializedPropertyChain.h"
 #include "WorldSnapshotData.h"
 
 bool FBaseClassDefaultArchive::ShouldSkipProperty(const FProperty* InProperty) const
diff --git a/Engine/Source/Editor/Experimental/RemoteExecution/Private/RemoteExecutionModule.cpp b/Engine/Source/Editor/Experimental/RemoteExecution/Private/RemoteExecutionModule.cpp
--- a/Engine/Source/Editor/Experimental/RemoteExecution/Private/RemoteExecutionModule.cpp
+++ b/Engine/Source/Editor/Experimental/RemoteExecution/Private/RemoteExecutionModule.cpp
@@ -5,6 +5,7 @@
 #include "ISettingsModule.h"
 #include "RemoteExecutionSettings.h"
 #include "Messages.h"
+#include "Modules/ModuleManager.h"
 
 IMPLEMENT_MODULE(FRemoteExecutionModule, EditorConfig);
 
