Fixes for no unity build errors
#rnx

diff --git a/Engine/Plugins/VirtualProduction/LevelSnapshots/Source/LevelSnapshots/Private/Data/LevelSnapshot.cpp b/Engine/Plugins/VirtualProduction/LevelSnapshots/Source/LevelSnapshots/Private/Data/LevelSnapshot.cpp
--- a/Engine/Plugins/VirtualProduction/LevelSnapshots/Source/LevelSnapshots/Private/Data/LevelSnapshot.cpp
+++ b/Engine/Plugins/VirtualProduction/LevelSnapshots/Source/LevelSnapshots/Private/Data/LevelSnapshot.cpp
@@ -87,7 +87,7 @@ void ULevelSnapshot::ApplySnapshotToWorld(UWorld* TargetWorld, const FPropertySe
 	
 	if (MapPath != FSoftObjectPath(TargetWorld))
 	{
-#if !(UE_BUILD_SHIPPING || UE_BUILD_TEST)
+#if WITH_EDITOR
 		if (TargetWorld->IsPlayInEditor())
 		{
 			FMessageLog("PIE").Warning(
@@ -97,7 +97,7 @@ void ULevelSnapshot::ApplySnapshotToWorld(UWorld* TargetWorld, const FPropertySe
 				)
 			);
 		}
-#endif
+#endif // WITH_EDITOR
 		UE_LOG(LogLevelSnapshots, Error, TEXT("This snapshot was taken for world '%s' and cannot be applied to world '%s': snapshots currently only support applying to the world they were taken in. "), *MapPath.ToString(), *TargetWorld->GetPathName());
 		return;
 	}
@@ -120,14 +120,14 @@ void ULevelSnapshot::SnapshotWorld(UWorld* TargetWorld)
 
 	if (TargetWorld->WorldType != EWorldType::Editor)
 	{
-#if !(UE_BUILD_SHIPPING || UE_BUILD_TEST)
+#if WITH_EDITOR
 		if (TargetWorld->IsPlayInEditor())
 		{
 			FMessageLog("PIE").Warning(
 			NSLOCTEXT("LevelSnapshots", "IncompatibleWorlds", "Taking snapshots in PIE is an experimental feature. The snapshot will work in the same PIE session but may no longer work when you start a new PIE session.")
 			);
 		}
-#endif
+#endif // WITH_EDITOR
 		UE_LOG(LogLevelSnapshots, Warning, TEXT("Level snapshots currently only support editors. Snapshots taken in other world types are experimental any may not function as expected."));
 	}
 	
diff --git a/Engine/Source/Developer/TraceServices/Private/TraceServicesModule.cpp b/Engine/Source/Developer/TraceServices/Private/TraceServicesModule.cpp
--- a/Engine/Source/Developer/TraceServices/Private/TraceServicesModule.cpp
+++ b/Engine/Source/Developer/TraceServices/Private/TraceServicesModule.cpp
@@ -15,6 +15,7 @@
 #include "Modules/DiagnosticsModule.h"
 #include "Modules/PlatformEventsModule.h"
 #include "Modules/TasksModule.h"
+#include "Stats/StatsTrace.h"
 
 class FTraceServicesModule
 	: public ITraceServicesModule
diff --git a/Engine/Source/Runtime/AVEncoder/Private/Encoders/NVENC/NVENC_Common.cpp b/Engine/Source/Runtime/AVEncoder/Private/Encoders/NVENC/NVENC_Common.cpp
--- a/Engine/Source/Runtime/AVEncoder/Private/Encoders/NVENC/NVENC_Common.cpp
+++ b/Engine/Source/Runtime/AVEncoder/Private/Encoders/NVENC/NVENC_Common.cpp
@@ -6,6 +6,7 @@
 #include "CoreMinimal.h"
 #include "VideoEncoderCommon.h"
 #include "HAL/PlatformProcess.h"
+#include "RHI.h"
 
 // define a function pointer for creating an instance of nvEncodeAPI
 typedef NVENCSTATUS(NVENCAPI* NVENCAPIPROC)(NV_ENCODE_API_FUNCTION_LIST*);
diff --git a/Engine/Source/Runtime/Engine/Private/LevelScriptActor.cpp b/Engine/Source/Runtime/Engine/Private/LevelScriptActor.cpp
--- a/Engine/Source/Runtime/Engine/Private/LevelScriptActor.cpp
+++ b/Engine/Source/Runtime/Engine/Private/LevelScriptActor.cpp
@@ -8,6 +8,7 @@
 #include "GameFramework/PlayerController.h"
 #include "Engine/InputDelegateBinding.h"
 #include "Engine/LevelScriptBlueprint.h"
+#include "UObject/UObjectHash.h"
 
 
 //////////////////////////////////////////////////////////////////////////
