VirtualProduction: Merge from RES.

Except RemoteControl due to a dependency.

diff --git a/Engine/Plugins/Compositing/Composure/Source/Composure/Private/CompositingElement.cpp b/Engine/Plugins/Compositing/Composure/Source/Composure/Private/CompositingElement.cpp
--- a/Engine/Plugins/Compositing/Composure/Source/Composure/Private/CompositingElement.cpp
+++ b/Engine/Plugins/Compositing/Composure/Source/Composure/Private/CompositingElement.cpp
@@ -701,7 +701,7 @@ ACameraActor* ACompositingElement::FindTargetCamera() const
 		for (TActorIterator<ACineCameraActor> CamActorIt(GetWorld()); CamActorIt; ++CamActorIt)
 		{
 			ACineCameraActor* FoundActor = *CamActorIt;
-			if (!FoundActor->IsPendingKill())
+			if (IsValid(FoundActor))
 			{
 				return FoundActor;
 			}
@@ -710,7 +710,7 @@ ACameraActor* ACompositingElement::FindTargetCamera() const
 		for (TActorIterator<ACameraActor> CamActorIt(GetWorld()); CamActorIt; ++CamActorIt)
 		{
 			ACameraActor* FoundActor = *CamActorIt;
-			if (!FoundActor->IsPendingKill())
+			if (IsValid(FoundActor))
 			{
 				return FoundActor;
 			}
diff --git a/Engine/Plugins/Developer/Concert/ConcertSync/ConcertSyncClient/Source/ConcertSyncClient/Private/ConcertClientTransactionBridge.cpp b/Engine/Plugins/Developer/Concert/ConcertSync/ConcertSyncClient/Source/ConcertSyncClient/Private/ConcertClientTransactionBridge.cpp
--- a/Engine/Plugins/Developer/Concert/ConcertSync/ConcertSyncClient/Source/ConcertSyncClient/Private/ConcertClientTransactionBridge.cpp
+++ b/Engine/Plugins/Developer/Concert/ConcertSync/ConcertSyncClient/Source/ConcertSyncClient/Private/ConcertClientTransactionBridge.cpp
@@ -153,7 +153,7 @@ struct FEditorTransactionNotification
 				DeltaChange.bHasNameChange = !InObjectUpdate.ObjectData.NewName.IsNone();
 				DeltaChange.bHasOuterChange = !InObjectUpdate.ObjectData.NewOuterPathName.IsNone();
 				DeltaChange.bHasExternalPackageChange = !InObjectUpdate.ObjectData.NewExternalPackageName.IsNone();
-				DeltaChange.bHasPendingKillChange = InObjectUpdate.ObjectData.bIsPendingKill != InTransactionObject->IsPendingKill();
+				DeltaChange.bHasPendingKillChange = InObjectUpdate.ObjectData.bIsPendingKill != !IsValid(InTransactionObject);
 				DeltaChange.bHasNonPropertyChanges = InObjectUpdate.ObjectData.SerializedData.Num() > 0;
 				for (const FConcertSerializedPropertyData& PropertyData : InObjectUpdate.PropertyDatas)
 				{
@@ -648,7 +648,7 @@ void FConcertClientTransactionBridge::HandleObjectTransacted(UObject* InObject,
 				ObjectUpdatePtr->ObjectId = ObjectId;
 				ObjectUpdatePtr->ObjectPathDepth = ConcertSyncClientUtil::GetObjectPathDepth(InObject);
 				ObjectUpdatePtr->ObjectData.bAllowCreate = false;
-				ObjectUpdatePtr->ObjectData.bIsPendingKill = InObject->IsPendingKill();
+				ObjectUpdatePtr->ObjectData.bIsPendingKill = !IsValid(InObject);
 			}
 
 			if (TransactionAnnotation.IsValid())
@@ -687,8 +687,8 @@ void FConcertClientTransactionBridge::HandleObjectTransacted(UObject* InObject,
 		FConcertExportedObject& ObjectUpdate = OngoingTransaction.FinalizedData.FinalizedObjectUpdates.AddDefaulted_GetRef();
 		ObjectUpdate.ObjectId = ObjectId;
 		ObjectUpdate.ObjectPathDepth = ConcertSyncClientUtil::GetObjectPathDepth(InObject);
-		ObjectUpdate.ObjectData.bAllowCreate = InTransactionEvent.HasPendingKillChange() && !InObject->IsPendingKill();
-		ObjectUpdate.ObjectData.bIsPendingKill = InObject->IsPendingKill();
+		ObjectUpdate.ObjectData.bAllowCreate = InTransactionEvent.HasPendingKillChange() && IsValid(InObject);
+		ObjectUpdate.ObjectData.bIsPendingKill = !IsValid(InObject);
 		ObjectUpdate.ObjectData.NewPackageName = NewObjectPackageName;
 		ObjectUpdate.ObjectData.NewName = NewObjectName;
 		ObjectUpdate.ObjectData.NewOuterPathName = NewObjectOuterPathName;
@@ -701,7 +701,7 @@ void FConcertClientTransactionBridge::HandleObjectTransacted(UObject* InObject,
 		}
 
 		// If this object changed from being pending kill to not being pending kill, we have to send a full object update (including all properties), rather than attempt a delta-update
-		const bool bForceFullObjectUpdate = InTransactionEvent.HasPendingKillChange() && !InObject->IsPendingKill();
+		const bool bForceFullObjectUpdate = InTransactionEvent.HasPendingKillChange() && IsValid(InObject);
 		if (bForceFullObjectUpdate)
 		{
 			// Serialize the entire object.
diff --git a/Engine/Plugins/Experimental/ColorCorrectRegions/Source/ColorCorrectRegions/Private/ColorCorrectRegion.cpp b/Engine/Plugins/Experimental/ColorCorrectRegions/Source/ColorCorrectRegions/Private/ColorCorrectRegion.cpp
--- a/Engine/Plugins/Experimental/ColorCorrectRegions/Source/ColorCorrectRegions/Private/ColorCorrectRegion.cpp
+++ b/Engine/Plugins/Experimental/ColorCorrectRegions/Source/ColorCorrectRegions/Private/ColorCorrectRegion.cpp
@@ -25,9 +25,9 @@ AColorCorrectRegion::AColorCorrectRegion(const FObjectInitializer& ObjectInitial
 void AColorCorrectRegion::BeginPlay()
 {	
 	Super::BeginPlay();
-	if (this->GetWorld())
+	if (const UWorld* World = GetWorld())
 	{
-		ColorCorrectRegionsSubsystem = Cast<UColorCorrectRegionsSubsystem>(this->GetWorld()->GetSubsystemBase(UColorCorrectRegionsSubsystem::StaticClass()));
+		ColorCorrectRegionsSubsystem = World->GetSubsystem<UColorCorrectRegionsSubsystem>();
 	}
 
 	if (ColorCorrectRegionsSubsystem)
@@ -85,6 +85,13 @@ void AColorCorrectRegion::PostEditChangeProperty(struct FPropertyChangedEvent& P
 	const FName PropertyName = PropertyChangedEvent.GetPropertyName();
 	if (PropertyName == GET_MEMBER_NAME_CHECKED(AColorCorrectRegion, Priority))
 	{
+		if (!ColorCorrectRegionsSubsystem)
+		{
+			if (const UWorld* World = GetWorld())
+			{
+				ColorCorrectRegionsSubsystem = World->GetSubsystem<UColorCorrectRegionsSubsystem>();
+			}
+		}
 		if (ColorCorrectRegionsSubsystem)
 		{
 			ColorCorrectRegionsSubsystem->SortRegionsByPriority();
diff --git a/Engine/Plugins/Experimental/ColorCorrectRegions/Source/ColorCorrectRegions/Private/ColorCorrectRegionsSceneViewExtension.cpp b/Engine/Plugins/Experimental/ColorCorrectRegions/Source/ColorCorrectRegions/Private/ColorCorrectRegionsSceneViewExtension.cpp
--- a/Engine/Plugins/Experimental/ColorCorrectRegions/Source/ColorCorrectRegions/Private/ColorCorrectRegionsSceneViewExtension.cpp
+++ b/Engine/Plugins/Experimental/ColorCorrectRegions/Source/ColorCorrectRegions/Private/ColorCorrectRegionsSceneViewExtension.cpp
@@ -370,7 +370,7 @@ void FColorCorrectRegionsSceneViewExtension::PrePostProcessPass_RenderThread(FRD
 			* Alternative is to get all component ids from ViewFamily.Scene and compare with actor's 
 			*	Region->ActiveMeshComponent->ComponentId ComponentIdSearchTable.Contains(RegionPrimitiveComponentId)
 			*/
-			if (Region->IsPendingKill() || 
+			if (!IsValid(Region) || 
 				Region->IsActorBeingDestroyed() ||
 				!Region->Enabled || 
 				Region->IsHidden() ||
diff --git a/Engine/Plugins/Experimental/VirtualCamera/Source/VCamCore/Private/VCamComponent.cpp b/Engine/Plugins/Experimental/VirtualCamera/Source/VCamCore/Private/VCamComponent.cpp
--- a/Engine/Plugins/Experimental/VirtualCamera/Source/VCamCore/Private/VCamComponent.cpp
+++ b/Engine/Plugins/Experimental/VirtualCamera/Source/VCamCore/Private/VCamComponent.cpp
@@ -132,7 +132,7 @@ void UVCamComponent::NotifyComponentWasReplaced(UVCamComponent* ReplacementCompo
 bool UVCamComponent::CanUpdate() const
 {
 	UWorld* World = GetWorld();
-	if (bEnabled && !IsPendingKill() && !bIsEditorObjectButPIEIsRunning && World)
+	if (bEnabled && IsValid(this) && !bIsEditorObjectButPIEIsRunning && World)
 	{
 		// Check for an Inactive type of world which means nothing should ever execute on this object
 		// @TODO: This is far from optimal as it means a zombie object has been created that never gets GC'ed
diff --git a/Engine/Plugins/Experimental/VirtualCamera/Source/VCamCoreEditor/Private/VCamEditorLibrary.cpp b/Engine/Plugins/Experimental/VirtualCamera/Source/VCamCoreEditor/Private/VCamEditorLibrary.cpp
--- a/Engine/Plugins/Experimental/VirtualCamera/Source/VCamCoreEditor/Private/VCamEditorLibrary.cpp
+++ b/Engine/Plugins/Experimental/VirtualCamera/Source/VCamCoreEditor/Private/VCamEditorLibrary.cpp
@@ -18,7 +18,7 @@ void UVCamEditorLibrary::GetAllVCamComponentsInLevel(TArray<UVCamComponent*>& VC
 	for (TObjectIterator<UVCamComponent> It(ExcludeFlags, true, EInternalObjectFlags::PendingKill); It; ++It)
 	{
 		UVCamComponent* VCamComponent = *It;
-		if (VCamComponent && !VCamComponent->IsPendingKill())
+		if (IsValid(VCamComponent))
 		{
 			// Ensure the object lives in the editor world
 			UWorld* World = VCamComponent->GetWorld();
diff --git a/Engine/Plugins/Experimental/VirtualCamera/Source/VirtualCamera/Private/VCamBlueprintFunctionLibrary.cpp b/Engine/Plugins/Experimental/VirtualCamera/Source/VirtualCamera/Private/VCamBlueprintFunctionLibrary.cpp
--- a/Engine/Plugins/Experimental/VirtualCamera/Source/VirtualCamera/Private/VCamBlueprintFunctionLibrary.cpp
+++ b/Engine/Plugins/Experimental/VirtualCamera/Source/VirtualCamera/Private/VCamBlueprintFunctionLibrary.cpp
@@ -11,11 +11,13 @@
 #if WITH_EDITOR
 #include "Editor.h"
 #include "EditorAssetLibrary.h"
+#include "ITakeRecorderModule.h"
 #include "LevelEditor.h"
 #include "LevelEditorViewport.h"
 #include "SceneView.h"
 #include "SLevelViewport.h"
 #include "LevelSequenceEditorBlueprintLibrary.h"
+#include "TakePreset.h"
 #include "VPUtilitiesEditor/Public/VPUtilitiesEditorBlueprintLibrary.h"
 #include "Subsystems/UnrealEditorSubsystem.h"
 #include "LevelEditor/Public/LevelEditorSubsystem.h"
@@ -38,7 +40,20 @@ UVirtualCameraUserSettings* UVCamBlueprintFunctionLibrary::GetUserSettings()
 ULevelSequence* UVCamBlueprintFunctionLibrary::GetCurrentLevelSequence()
 {
 #if WITH_EDITOR
-	return ULevelSequenceEditorBlueprintLibrary::GetCurrentLevelSequence();
+	ULevelSequence* LevelSequence = ULevelSequenceEditorBlueprintLibrary::GetCurrentLevelSequence();
+
+	//if there's no sequence open in sequencer, use take recorder's pending take
+	if (!LevelSequence)
+	{
+		ITakeRecorderModule& TakeRecorderModule = FModuleManager::LoadModuleChecked<ITakeRecorderModule>("TakeRecorder");
+		UTakePreset* TakePreset = TakeRecorderModule.GetPendingTake();
+		
+		if (TakePreset)
+		{
+			LevelSequence = TakePreset->GetLevelSequence();
+		}
+	}
+	return LevelSequence;
 #else
 	return nullptr;
 #endif
diff --git a/Engine/Plugins/Experimental/VirtualProductionUtilities/Source/VPBookmark/Private/VPBookmarkBlueprintLibrary.cpp b/Engine/Plugins/Experimental/VirtualProductionUtilities/Source/VPBookmark/Private/VPBookmarkBlueprintLibrary.cpp
--- a/Engine/Plugins/Experimental/VirtualProductionUtilities/Source/VPBookmark/Private/VPBookmarkBlueprintLibrary.cpp
+++ b/Engine/Plugins/Experimental/VirtualProductionUtilities/Source/VPBookmark/Private/VPBookmarkBlueprintLibrary.cpp
@@ -57,7 +57,7 @@ void UVPBookmarkBlueprintLibrary::GetAllVPBookmarkActors(const UObject* WorldCon
 				{
 					if (AActor* Actor = Bookmark->OwnedActor.Get())
 					{
-						if (!Bookmark->IsPendingKill() && !Actor->IsPendingKill())
+						if (IsValid(Bookmark) && IsValid(Actor))
 						{
 							OutActors.Add(Actor);
 						}
@@ -85,7 +85,7 @@ void UVPBookmarkBlueprintLibrary::GetAllVPBookmark(const UObject* WorldContextOb
 			{
 				if (UVPBookmark* Bookmark = Cast<UVPBookmark>(Base))
 				{
-					if (!Bookmark->IsPendingKill())
+					if (IsValid(Bookmark))
 					{
 						OutBookmarks.Add(Bookmark);
 					}
diff --git a/Engine/Plugins/Experimental/VirtualProductionUtilities/Source/VPUtilities/Private/VPFullScreenUserWidget.cpp b/Engine/Plugins/Experimental/VirtualProductionUtilities/Source/VPUtilities/Private/VPFullScreenUserWidget.cpp
--- a/Engine/Plugins/Experimental/VirtualProductionUtilities/Source/VPUtilities/Private/VPFullScreenUserWidget.cpp
+++ b/Engine/Plugins/Experimental/VirtualProductionUtilities/Source/VPUtilities/Private/VPFullScreenUserWidget.cpp
@@ -53,6 +53,75 @@ namespace
 			return EVisibility::SelfHitTestInvisible;
 		}
 	}
+
+	namespace VPFullScreenUserWidgetPrivate
+	{
+		/**
+		 * Class made to handle world cleanup and hide/cleanup active UserWidget to avoid touching public headers
+		 */
+		class FWorldCleanupListener
+		{
+		public:
+
+			static FWorldCleanupListener* Get()
+			{
+				static FWorldCleanupListener Instance;
+				return &Instance;
+			}
+
+			/** Disallow Copying / Moving */
+			UE_NONCOPYABLE(FWorldCleanupListener);
+
+			~FWorldCleanupListener()
+			{
+				FWorldDelegates::OnWorldCleanup.RemoveAll(this);
+			}
+
+			void AddWidget(UVPFullScreenUserWidget* InWidget)
+			{
+				WidgetsToHide.AddUnique(InWidget);
+			}
+
+			void RemoveWidget(UVPFullScreenUserWidget* InWidget)
+			{
+				WidgetsToHide.RemoveSingleSwap(InWidget, false);
+			}
+
+		private:
+
+			FWorldCleanupListener()
+			{
+				FWorldDelegates::OnWorldCleanup.AddRaw(this, &FWorldCleanupListener::OnWorldCleanup);
+			}
+
+			void OnWorldCleanup(UWorld* InWorld, bool bSessionEnded, bool bCleanupResources)
+			{
+				for (auto WeakWidgetIter = WidgetsToHide.CreateIterator(); WeakWidgetIter; ++WeakWidgetIter)
+				{
+					TWeakObjectPtr<UVPFullScreenUserWidget>& WeakWidget = *WeakWidgetIter;
+					if (UVPFullScreenUserWidget* Widget = WeakWidget.Get())
+					{
+						if (Widget->IsDisplayed()
+							&& Widget->GetWidget()
+							&& (Widget->GetWidget()->GetWorld() == InWorld))
+						{
+							//Remove first since Hide removes object from the list
+							WeakWidgetIter.RemoveCurrent();
+							Widget->Hide();
+						}
+					}
+					else
+					{
+						WeakWidgetIter.RemoveCurrent();
+					}
+				}
+			}
+
+		private:
+
+			TArray<TWeakObjectPtr<UVPFullScreenUserWidget>> WidgetsToHide;
+		};
+	}
 }
 
 
@@ -644,6 +713,8 @@ bool UVPFullScreenUserWidget::Display(UWorld* InWorld)
 			FWorldDelegates::LevelRemovedFromWorld.AddUObject(this, &UVPFullScreenUserWidget::OnLevelRemovedFromWorld);
 			FWorldDelegates::OnWorldCleanup.AddUObject(this, &UVPFullScreenUserWidget::OnWorldCleanup);
 
+			VPFullScreenUserWidgetPrivate::FWorldCleanupListener::Get()->AddWidget(this);
+
 			// If we are using Composure as our output, then send the WidgetRenderTarget to each one
 			if (CurrentDisplayType == EVPWidgetDisplayType::Composure)
 			{
@@ -691,6 +762,8 @@ void UVPFullScreenUserWidget::Hide()
 		FWorldDelegates::LevelRemovedFromWorld.RemoveAll(this);
 		FWorldDelegates::OnWorldCleanup.RemoveAll(this);
 
+		VPFullScreenUserWidgetPrivate::FWorldCleanupListener::Get()->RemoveWidget(this);
+
 		if (CurrentDisplayType == EVPWidgetDisplayType::Viewport)
 		{
 			ViewportDisplayType.Hide(World.Get());
diff --git a/Engine/Plugins/Experimental/VirtualProductionUtilities/Source/VPUtilitiesEditor/Private/VPUtilitiesEditorModule.cpp b/Engine/Plugins/Experimental/VirtualProductionUtilities/Source/VPUtilitiesEditor/Private/VPUtilitiesEditorModule.cpp
--- a/Engine/Plugins/Experimental/VirtualProductionUtilities/Source/VPUtilitiesEditor/Private/VPUtilitiesEditorModule.cpp
+++ b/Engine/Plugins/Experimental/VirtualProductionUtilities/Source/VPUtilitiesEditor/Private/VPUtilitiesEditorModule.cpp
@@ -165,7 +165,7 @@ void FVPUtilitiesEditorModule::InitializeOSCServer()
 	}
 	else
 	{
-		OSCServer.Reset(UOSCManager::CreateOSCServer(ServerAddress, ServerPort, false, true, FString()));
+		OSCServer.Reset(UOSCManager::CreateOSCServer(ServerAddress, ServerPort, false, true, FString(), GetTransientPackage()));
 		
 #if WITH_EDITOR
 		// Allow it to tick in editor, so that messages are parsed.
diff --git a/Engine/Plugins/VirtualProduction/CameraCalibration/Source/CameraCalibrationEditor/Private/AssetEditor/CameraCalibrationStepsController.cpp b/Engine/Plugins/VirtualProduction/CameraCalibration/Source/CameraCalibrationEditor/Private/AssetEditor/CameraCalibrationStepsController.cpp
--- a/Engine/Plugins/VirtualProduction/CameraCalibration/Source/CameraCalibrationEditor/Private/AssetEditor/CameraCalibrationStepsController.cpp
+++ b/Engine/Plugins/VirtualProduction/CameraCalibration/Source/CameraCalibrationEditor/Private/AssetEditor/CameraCalibrationStepsController.cpp
@@ -516,7 +516,7 @@ void FCameraCalibrationStepsController::CreateComp()
 						break;
 					}
 				}
-		}
+			}
 		}
 
 		// Break since don't need to look at more UMediaTextureCompositingInputs.
@@ -579,7 +579,7 @@ void FCameraCalibrationStepsController::CreateComp()
 		return;
 	}
 
-	RenderTarget->RenderTargetFormat = RTF_RGBA8;
+	RenderTarget->RenderTargetFormat = RTF_RGBA16f;
 	RenderTarget->ClearColor = FLinearColor::Black;
 	RenderTarget->bAutoGenerateMips = true;
 	RenderTarget->InitAutoFormat(1920, 1080);
@@ -732,12 +732,12 @@ void FCameraCalibrationStepsController::EnableDistortionInCG()
 		// the CG layer will automatically pick it up and start using it.
 		if (DistortionHandler)
 		{
-		CaptureBase->SetDistortionHandler(DistortionHandler);
-	}
+			CaptureBase->SetDistortionHandler(DistortionHandler);
+		}
 		else
 		{
 			UE_LOG(LogCameraCalibrationEditor, Warning, TEXT("Could not find a distortion handler in the selected camera"));
-}
+		}
 	}
 }
 
diff --git a/Engine/Plugins/VirtualProduction/CameraCalibration/Source/CameraCalibrationEditor/Private/AssetEditor/SSimulcamViewport.cpp b/Engine/Plugins/VirtualProduction/CameraCalibration/Source/CameraCalibrationEditor/Private/AssetEditor/SSimulcamViewport.cpp
--- a/Engine/Plugins/VirtualProduction/CameraCalibration/Source/CameraCalibrationEditor/Private/AssetEditor/SSimulcamViewport.cpp
+++ b/Engine/Plugins/VirtualProduction/CameraCalibration/Source/CameraCalibrationEditor/Private/AssetEditor/SSimulcamViewport.cpp
@@ -59,7 +59,7 @@ void SSimulcamViewport::Construct(const FArguments& InArgs, UTexture* InTexture)
 				Opacity.MaskA = 1;
 			}
 
-			Material->BlendMode = BLEND_AlphaComposite;
+			Material->BlendMode = BLEND_Opaque;
 
 			Material->Expressions.Add(TextureSampler);
 			Material->MaterialDomain = EMaterialDomain::MD_UI;
diff --git a/Engine/Plugins/VirtualProduction/DMX/DMXEngine/Source/DMXBlueprintGraph/Private/Widgets/SDMXEntityReferenceGraphPin.cpp b/Engine/Plugins/VirtualProduction/DMX/DMXEngine/Source/DMXBlueprintGraph/Private/Widgets/SDMXEntityReferenceGraphPin.cpp
--- a/Engine/Plugins/VirtualProduction/DMX/DMXEngine/Source/DMXBlueprintGraph/Private/Widgets/SDMXEntityReferenceGraphPin.cpp
+++ b/Engine/Plugins/VirtualProduction/DMX/DMXEngine/Source/DMXBlueprintGraph/Private/Widgets/SDMXEntityReferenceGraphPin.cpp
@@ -95,7 +95,7 @@ FDMXEntityReference SDMXEntityReferenceGraphPin::GetPinValue() const
 	{
 		FDMXEntityReference::StaticStruct()->ImportText(*EntityRefStr, &EntityReference, nullptr, EPropertyPortFlags::PPF_None, GLog, FDMXEntityReference::StaticStruct()->GetName());
 	}
-	else if (!GraphPinObj->IsPendingKill() && !GraphPinObj->GetSchema()->IsPendingKill() && GraphPinObj->GetDefaultAsString().IsEmpty())
+	else if (!GraphPinObj->IsPendingKill() && IsValid(GraphPinObj->GetSchema()) && GraphPinObj->GetDefaultAsString().IsEmpty())
 	{
 		// Defaults are empty and, no matter the type of Entity Ref, we would return UDMXEntity as type because
 		// we are creating a FDMXEntityReference to return.
diff --git a/Engine/Plugins/VirtualProduction/LevelSnapshots/Source/LevelSnapshots/LevelSnapshots.Build.cs b/Engine/Plugins/VirtualProduction/LevelSnapshots/Source/LevelSnapshots/LevelSnapshots.Build.cs
--- a/Engine/Plugins/VirtualProduction/LevelSnapshots/Source/LevelSnapshots/LevelSnapshots.Build.cs
+++ b/Engine/Plugins/VirtualProduction/LevelSnapshots/Source/LevelSnapshots/LevelSnapshots.Build.cs
@@ -31,6 +31,7 @@ public class LevelSnapshots : ModuleRules
             new string[]
             {
                 "Core",
+                "TraceLog"
             }
             );
 
diff --git a/Engine/Plugins/VirtualProduction/LevelSnapshots/Source/LevelSnapshots/Private/ApplySnapshotFilter.cpp b/Engine/Plugins/VirtualProduction/LevelSnapshots/Source/LevelSnapshots/Private/ApplySnapshotFilter.cpp
--- a/Engine/Plugins/VirtualProduction/LevelSnapshots/Source/LevelSnapshots/Private/ApplySnapshotFilter.cpp
+++ b/Engine/Plugins/VirtualProduction/LevelSnapshots/Source/LevelSnapshots/Private/ApplySnapshotFilter.cpp
@@ -60,6 +60,8 @@ FApplySnapshotFilter FApplySnapshotFilter::Make(ULevelSnapshot* Snapshot, AActor
 
 void FApplySnapshotFilter::ApplyFilterToFindSelectedProperties(FPropertySelectionMap& MapToAddTo)
 {
+	SCOPED_SNAPSHOT_CORE_TRACE(ApplyFilters);
+	
 	if (EnsureParametersAreValid() && FSnapshotRestorability::IsActorDesirableForCapture(WorldActor) && EFilterResult::CanInclude(Filter->IsActorValid({ DeserializedSnapshotActor, WorldActor })))
 	{
 		FilterActorPair(MapToAddTo);
@@ -98,14 +100,15 @@ bool FApplySnapshotFilter::EnsureParametersAreValid() const
 
 	UClass* WorldClass = WorldActor->GetClass();
 	UClass* DeserializedClass = DeserializedSnapshotActor->GetClass();
-	if (!ensure(WorldClass == DeserializedClass))
+	if (WorldClass != DeserializedClass)
 	{
 		UE_LOG(
             LogLevelSnapshots,
             Error,
-            TEXT("FApplySnapshotFilter::ApplyFilterToFindSelectedProperties: WorldActor class '%s' differs from DerserializedSnapshoatActor class '%s'. Differing classes are not supported."),
+            TEXT("FApplySnapshotFilter::ApplyFilterToFindSelectedProperties: WorldActor class '%s' differs from DerserializedSnapshoatActor class '%s'. Differing classes are not supported. (Actor: %s)"),
             *WorldActor->GetClass()->GetName(),
-            *DeserializedSnapshotActor->GetClass()->GetName()
+            *DeserializedSnapshotActor->GetClass()->GetName(),
+			*WorldActor->GetPathName()
         );
 		return false;
 	}
diff --git a/Engine/Plugins/VirtualProduction/LevelSnapshots/Source/LevelSnapshots/Private/Archive/ApplySnapshotDataArchiveV2.cpp b/Engine/Plugins/VirtualProduction/LevelSnapshots/Source/LevelSnapshots/Private/Archive/ApplySnapshotDataArchiveV2.cpp
--- a/Engine/Plugins/VirtualProduction/LevelSnapshots/Source/LevelSnapshots/Private/Archive/ApplySnapshotDataArchiveV2.cpp
+++ b/Engine/Plugins/VirtualProduction/LevelSnapshots/Source/LevelSnapshots/Private/Archive/ApplySnapshotDataArchiveV2.cpp
@@ -4,8 +4,6 @@
 
 #include "LevelSnapshotsLog.h"
 #include "Data/PropertySelection.h"
-#include "LevelSnapshotsStats.h"
-#include "Archive/ClassDefaults/TakeClassDefaultObjectSnapshotArchive.h"
 #include "WorldSnapshotData.h"
 #include "ClassDefaults/ApplyClassDefaulDataArchive.h"
 
@@ -180,7 +178,7 @@ void FApplySnapshotDataArchiveV2::ApplyToRecreatedEditorWorldObject(FObjectSnaps
 
 bool FApplySnapshotDataArchiveV2::ShouldSkipProperty(const FProperty* InProperty) const
 {
-	DECLARE_SCOPE_CYCLE_COUNTER(TEXT("ShouldSkipProperty_Loading"), STAT_ShouldSkipProperty_Loading, STATGROUP_LevelSnapshots);
+	SCOPED_SNAPSHOT_CORE_TRACE(ShouldSkipProperty);
 	
 	bool bShouldSkipProperty = Super::ShouldSkipProperty(InProperty);
 	
diff --git a/Engine/Plugins/VirtualProduction/LevelSnapshots/Source/LevelSnapshots/Private/Data/ActorSnapshotData.cpp b/Engine/Plugins/VirtualProduction/LevelSnapshots/Source/LevelSnapshots/Private/Data/ActorSnapshotData.cpp
--- a/Engine/Plugins/VirtualProduction/LevelSnapshots/Source/LevelSnapshots/Private/Data/ActorSnapshotData.cpp
+++ b/Engine/Plugins/VirtualProduction/LevelSnapshots/Source/LevelSnapshots/Private/Data/ActorSnapshotData.cpp
@@ -93,8 +93,8 @@ TOptional<AActor*> FActorSnapshotData::GetPreallocated(UWorld* SnapshotWorld, FW
 	if (!CachedSnapshotActor.IsValid())
 	{
 		const FSoftClassPath SoftClassPath(ActorClass);
-		UClass* TargetClass = SoftClassPath.ResolveClass();
-		if (!ensureAlways(TargetClass))
+		UClass* TargetClass = SoftClassPath.TryLoadClass<AActor>();
+		if (!TargetClass)
 		{
 			UE_LOG(LogLevelSnapshots, Error, TEXT("Unknown class %s. The snapshot is mostly likely referencing a class that was deleted."), *SoftClassPath.ToString());
 			return {};
diff --git a/Engine/Plugins/VirtualProduction/LevelSnapshots/Source/LevelSnapshots/Private/Data/CustomSerialization/CustomObjectSerializationWrapper.cpp b/Engine/Plugins/VirtualProduction/LevelSnapshots/Source/LevelSnapshots/Private/Data/CustomSerialization/CustomObjectSerializationWrapper.cpp
--- a/Engine/Plugins/VirtualProduction/LevelSnapshots/Source/LevelSnapshots/Private/Data/CustomSerialization/CustomObjectSerializationWrapper.cpp
+++ b/Engine/Plugins/VirtualProduction/LevelSnapshots/Source/LevelSnapshots/Private/Data/CustomSerialization/CustomObjectSerializationWrapper.cpp
@@ -4,7 +4,6 @@
 
 #include "ApplySnapshotDataArchiveV2.h"
 #include "LevelSnapshotsLog.h"
-#include "LevelSnapshotsStats.h"
 #include "CustomSerialization/CustomSerializationDataManager.h"
 #include "LevelSnapshotsModule.h"
 #include "PropertySelectionMap.h"
@@ -26,6 +25,8 @@ namespace
 		FSerializationDataGetter SerializationDataGetter
 		)
 	{
+		SCOPED_SNAPSHOT_CORE_TRACE(CustomObjectSerialization_PostSnapshotRestore);
+		
 		FLevelSnapshotsModule& LevelSnapshots = FModuleManager::Get().GetModuleChecked<FLevelSnapshotsModule>("LevelSnapshots");
 		TSharedPtr<ICustomObjectSnapshotSerializer> CustomSerializer = LevelSnapshots.GetCustomSerializerForClass(SnapshotObject->GetClass());
 		if (!CustomSerializer.IsValid() || !ensure(SerializationDataGetter() != nullptr))
@@ -46,13 +47,10 @@ namespace
 					UE_LOG(LogLevelSnapshots, Warning, TEXT("FindOrRecreateSubobjectInSnapshotWorld did not return any valid subobject. Skipping subobject restoration..."));
 					return;
 				}
-				
+
+				// Recursively check whether subobjects also have a registered ICustomObjectSnapshotSerializer
+				const FRestoreObjectScope FinishRestore = PreObjectRestore_SnapshotWorld(SnapshotSubobject, WorldData, LocalisationSnapshotPackage, [&WorldData, OriginalPath = MetaData->GetOriginalPath()](){ return WorldData.GetCustomSubobjectData_ForSubobject(OriginalPath);});
 				FSnapshotArchive::ApplyToSnapshotWorldObject(SerializationDataGetter()->Subobjects[i], WorldData, SnapshotSubobject, LocalisationSnapshotPackage);
-				{
-					// Recursively check whether subobjects also have a registered ICustomObjectSnapshotSerializer
-					const FRestoreObjectScope FinishRestore = PreObjectRestore_SnapshotWorld(SnapshotSubobject, WorldData, LocalisationSnapshotPackage, [&WorldData, OriginalPath = MetaData->GetOriginalPath()](){ return WorldData.GetCustomSubobjectData_ForSubobject(OriginalPath);});
-				}
-				
 				CustomSerializer->OnPostSerializeSnapshotSubobject(SnapshotSubobject, *MetaData, SerializationDataReader);
 			}
 
@@ -69,6 +67,8 @@ namespace
 		FSerializationDataGetter SerializationDataGetter
 		)
 	{
+		SCOPED_SNAPSHOT_CORE_TRACE(CustomObjectSerialization_PostEditorRestore);
+		
 		FLevelSnapshotsModule& LevelSnapshots = FModuleManager::Get().GetModuleChecked<FLevelSnapshotsModule>("LevelSnapshots");
 		TSharedPtr<ICustomObjectSnapshotSerializer> CustomSerializer = LevelSnapshots.GetCustomSerializerForClass(EditorObject->GetClass());
 		if (!CustomSerializer.IsValid())
@@ -83,28 +83,31 @@ namespace
 			for (int32 i = 0; i < SerializationDataReader.GetNumSubobjects(); ++i)
 			{	
 				const TSharedPtr<ISnapshotSubobjectMetaData> MetaData = SerializationDataReader.GetSubobjectMetaData(i);
-				UObject* EditorSubobject = CustomSerializer->FindOrRecreateSubobjectInSnapshotWorld(EditorObject, *MetaData, SerializationDataReader);
+				UObject* EditorSubobject = CustomSerializer->FindOrRecreateSubobjectInEditorWorld(EditorObject, *MetaData, SerializationDataReader);
 				UObject* SnapshotSubobject = CustomSerializer->FindOrRecreateSubobjectInSnapshotWorld(SnapshotObject, *MetaData, SerializationDataReader);
-				if (!EditorSubobject || !SnapshotSubobject || !ensure(EditorSubobject->IsIn(EditorObject) && SnapshotSubobject->IsIn(SnapshotObject)))
+				if (!SnapshotSubobject || !ensure(SnapshotSubobject->IsIn(SnapshotObject)))
 				{
-					UE_LOG(LogLevelSnapshots, Warning, TEXT("FindOrRecreateSubobjectInSnapshotWorld did not return any valid subobject. Skipping subobject restoration..."));
+					UE_LOG(LogLevelSnapshots, Warning, TEXT("FindOrRecreateSubobjectInSnapshotWorld did not return any valid snapshot subobject. Skipping this subobject for %s"), *EditorObject->GetPathName());
+					continue;
+				}
+				if (!EditorSubobject || !ensure(EditorSubobject->IsIn(EditorObject)))
+				{
+					UE_LOG(LogLevelSnapshots, Warning, TEXT("FindOrRecreateSubobjectInSnapshotWorld did not return any valid editor subobject. Skipping this subobject for %s"), *EditorObject->GetPathName());
 					continue;
 				}
 
 				if (const FPropertySelection* SelectedProperties = SelectionMap.GetSelectedProperties(EditorSubobject))
 				{
+					// Recursively check whether subobjects also have a registered ICustomObjectSnapshotSerializer
+					const FRestoreObjectScope FinishRestore = PreObjectRestore_EditorWorld(SnapshotSubobject, EditorSubobject, WorldData, SelectionMap, LocalisationSnapshotPackage, [&WorldData, OriginalPath = MetaData->GetOriginalPath()](){ return WorldData.GetCustomSubobjectData_ForSubobject(OriginalPath);} );
+				
 					FCustomSerializationData* SerializationData = SerializationDataGetter();
 					FApplySnapshotDataArchiveV2::ApplyToExistingEditorWorldObject(SerializationData->Subobjects[i], WorldData, EditorSubobject, SnapshotSubobject, SelectionMap, *SelectedProperties);
 					CustomSerializer->OnPostSerializeEditorSubobject(EditorSubobject, *MetaData, SerializationDataReader);
-					// Recursively check whether subobjects also have a registered ICustomObjectSnapshotSerializer
-					const FRestoreObjectScope FinishRestore = PreObjectRestore_EditorWorld(SnapshotSubobject, EditorSubobject, WorldData, SelectionMap, LocalisationSnapshotPackage, [&WorldData, OriginalPath = MetaData->GetOriginalPath()](){ return WorldData.GetCustomSubobjectData_ForSubobject(OriginalPath);} );
-				}
-				else
-				{
-					UE_LOG(LogLevelSnapshots, Verbose, TEXT("There were no changes to the custom serialized subobject"));	
+					continue;
 				}
 				
-				CustomSerializer->OnPostSerializeSnapshotSubobject(EditorSubobject, *MetaData, SerializationDataReader);
+				UE_LOG(LogLevelSnapshots, Warning, TEXT("Editor subobject %s was not restored"), *EditorSubobject->GetPathName());	
 			}
 
 			CustomSerializer->PostApplySnapshotProperties(EditorObject, SerializationDataReader);
@@ -117,7 +120,7 @@ void FCustomObjectSerializationWrapper::TakeSnapshotForActor(
 	FCustomSerializationData& ActorSerializationData,
 	FWorldSnapshotData& WorldData)
 {
-	DECLARE_SCOPE_CYCLE_COUNTER(TEXT("TakeSnapshotForActor"), STAT_TakeSnapshotForActor, STATGROUP_LevelSnapshots);
+	SCOPED_SNAPSHOT_CORE_TRACE(CustomObjectSerialization_TakeActorSnapshot);
 	
 	FLevelSnapshotsModule& LevelSnapshots = FModuleManager::Get().GetModuleChecked<FLevelSnapshotsModule>("LevelSnapshots");
 	TSharedPtr<ICustomObjectSnapshotSerializer> CustomSerializer = LevelSnapshots.GetCustomSerializerForClass(EditorActor->GetClass());
@@ -138,7 +141,7 @@ void FCustomObjectSerializationWrapper::TakeSnapshotForSubobject(
 	UObject* Subobject,
 	FWorldSnapshotData& WorldData)
 {
-	DECLARE_SCOPE_CYCLE_COUNTER(TEXT("TakeSnapshotForSubobject"), STAT_TakeSnapshotForSubobject, STATGROUP_LevelSnapshots);
+	SCOPED_SNAPSHOT_CORE_TRACE(CustomObjectSerialization_TakeSubobjectSnapshot);
 	
 	FLevelSnapshotsModule& LevelSnapshots = FModuleManager::Get().GetModuleChecked<FLevelSnapshotsModule>("LevelSnapshots");
 	TSharedPtr<ICustomObjectSnapshotSerializer> CustomSerializer = LevelSnapshots.GetCustomSerializerForClass(Subobject->GetClass());
@@ -162,7 +165,8 @@ FRestoreObjectScope FCustomObjectSerializationWrapper::PreActorRestore_SnapshotW
 	FWorldSnapshotData& WorldData,
 	UPackage* LocalisationSnapshotPackage)
 {
-	DECLARE_SCOPE_CYCLE_COUNTER(TEXT("StartRestoreSnapshotForActor_SnapshotWorld"), STAT_StartRestoreSnapshotForActor_SnapshotWorld, STATGROUP_LevelSnapshots);
+	SCOPED_SNAPSHOT_CORE_TRACE(CustomObjectSerialization_PreSnapshotActorRestore);
+	
 	return PreObjectRestore_SnapshotWorld(
 		SnapshotActor,
 		WorldData,
@@ -178,8 +182,8 @@ FRestoreObjectScope FCustomObjectSerializationWrapper::PreActorRestore_EditorWor
 	const FPropertySelectionMap& SelectionMap,
 	UPackage* LocalisationSnapshotPackage)
 {
-	DECLARE_SCOPE_CYCLE_COUNTER(TEXT("StartRestoreSnapshotForActor_EditorWorld"), STAT_StartRestoreSnapshotForActor_EditorWorld, STATGROUP_LevelSnapshots);
-
+	SCOPED_SNAPSHOT_CORE_TRACE(CustomObjectSerialization_PreEditorRestore);
+	
 	const TOptional<AActor*> SnapshotActor = WorldData.GetDeserializedActor(EditorActor, LocalisationSnapshotPackage);
 	check(SnapshotActor);
 	
@@ -200,7 +204,8 @@ FRestoreObjectScope FCustomObjectSerializationWrapper::PreSubobjectRestore_Snaps
 	UPackage* LocalisationSnapshotPackage
 	)
 {
-	DECLARE_SCOPE_CYCLE_COUNTER(TEXT("StartRestoreSnapshotForSubobject_SnapshotWorld"), STAT_StartRestoreSnapshotForSubobject_SnapshotWorld, STATGROUP_LevelSnapshots);
+	SCOPED_SNAPSHOT_CORE_TRACE(CustomObjectSerialization_PreSnapshotRestore);
+	
 	return PreObjectRestore_SnapshotWorld(
 		Subobject,
 		WorldData,
@@ -217,7 +222,7 @@ FRestoreObjectScope FCustomObjectSerializationWrapper::PreSubobjectRestore_Edito
 	UPackage* LocalisationSnapshotPackage
 	)
 {
-	DECLARE_SCOPE_CYCLE_COUNTER(TEXT("StartRestoreSnapshotForSubobject_EditorWorld"), STAT_StartRestoreSnapshotForSubobject_EditorWorld, STATGROUP_LevelSnapshots);
+	SCOPED_SNAPSHOT_CORE_TRACE(CustomObjectSerialization_PreEditorRestore);
 	
 	const FSoftObjectPath SubobjectPath(EditorObject);
 	if (!WorldData.GetCustomSubobjectData_ForSubobject(SubobjectPath))
diff --git a/Engine/Plugins/VirtualProduction/LevelSnapshots/Source/LevelSnapshots/Private/Data/LevelSnapshot.cpp b/Engine/Plugins/VirtualProduction/LevelSnapshots/Source/LevelSnapshots/Private/Data/LevelSnapshot.cpp
--- a/Engine/Plugins/VirtualProduction/LevelSnapshots/Source/LevelSnapshots/Private/Data/LevelSnapshot.cpp
+++ b/Engine/Plugins/VirtualProduction/LevelSnapshots/Source/LevelSnapshots/Private/Data/LevelSnapshot.cpp
@@ -4,7 +4,6 @@
 
 #include "LevelSnapshotsLog.h"
 #include "LevelSnapshotsModule.h"
-#include "LevelSnapshotsStats.h"
 #include "PropertyInfoHelpers.h"
 #include "Restorability/PropertyComparisonParams.h"
 
@@ -16,6 +15,7 @@
 #include "SnapshotRestorability.h"
 #include "CustomSerialization/CustomObjectSerializationWrapper.h"
 #include "CustomSerialization/CustomSerializationDataManager.h"
+#include "Misc/ScopeExit.h"
 #include "UObject/Package.h"
 #include "UObject/TextProperty.h"
 #include "UObject/UnrealType.h"
@@ -95,6 +95,7 @@ namespace
 
 void ULevelSnapshot::ApplySnapshotToWorld(UWorld* TargetWorld, const FPropertySelectionMap& SelectionSet)
 {
+	SCOPED_SNAPSHOT_CORE_TRACE(ApplyToWorld);
 	if (TargetWorld == nullptr)
 	{
 		return;
@@ -116,8 +117,15 @@ void ULevelSnapshot::ApplySnapshotToWorld(UWorld* TargetWorld, const FPropertySe
 		UE_LOG(LogLevelSnapshots, Error, TEXT("This snapshot was taken for world '%s' and cannot be applied to world '%s': snapshots currently only support applying to the world they were taken in. "), *MapPath.ToString(), *TargetWorld->GetPathName());
 		return;
 	}
+
+	UE_LOG(LogLevelSnapshots, Log, TEXT("Applying snapshot %s to world %s"), *GetPathName(), *TargetWorld->GetPathName());
+	ON_SCOPE_EXIT
+	{
+		UE_LOG(LogLevelSnapshots, Log, TEXT("Finished applying snapshot"));
+	};
 	
 	EnsureWorldInitialised();
+	
 #if WITH_EDITOR
 	FScopedTransaction Transaction(FText::FromString("Loading Level Snapshot."));
 #endif
@@ -126,7 +134,8 @@ void ULevelSnapshot::ApplySnapshotToWorld(UWorld* TargetWorld, const FPropertySe
 
 bool ULevelSnapshot::SnapshotWorld(UWorld* TargetWorld)
 {
-	DECLARE_SCOPE_CYCLE_COUNTER(TEXT("TakeLevelSnapshot"), STAT_TakeLevelSnapshot, STATGROUP_LevelSnapshots);
+	SCOPED_SNAPSHOT_CORE_TRACE(SnapshotWorld);
+	
 	if (!ensure(TargetWorld))
 	{
 		UE_LOG(LogLevelSnapshots, Warning, TEXT("Unable To Snapshot World as World was invalid"));
@@ -166,7 +175,7 @@ bool ULevelSnapshot::SnapshotWorld(UWorld* TargetWorld)
 
 bool ULevelSnapshot::HasOriginalChangedPropertiesSinceSnapshotWasTaken(AActor* SnapshotActor, AActor* WorldActor) const
 {
-	DECLARE_SCOPE_CYCLE_COUNTER(TEXT("HasOriginalChangedSinceSnapshot"), STAT_AreAllSnapshotRelevantPropertiesIdentical, STATGROUP_LevelSnapshots);
+	SCOPED_SNAPSHOT_CORE_TRACE(HasOriginalChangedSinceSnapshot);
 
 	if (!IsValid(SnapshotActor) || !IsValid(WorldActor))
 	{
@@ -307,7 +316,7 @@ bool ULevelSnapshot::AreSnapshotAndOriginalPropertiesEquivalent(const FProperty*
 
 TOptional<AActor*> ULevelSnapshot::GetDeserializedActor(const FSoftObjectPath& OriginalActorPath)
 {
-	DECLARE_SCOPE_CYCLE_COUNTER(TEXT("GetDeserializedActor"), STAT_GetDeserializedActor, STATGROUP_LevelSnapshots);
+	SCOPED_SNAPSHOT_CORE_TRACE(GetDeserializedActor);
 	
 	EnsureWorldInitialised();
 	return SerializedData.GetDeserializedActor(OriginalActorPath, GetPackage());
@@ -320,48 +329,80 @@ int32 ULevelSnapshot::GetNumSavedActors() const
 
 void ULevelSnapshot::DiffWorld(UWorld* World, FActorPathConsumer HandleMatchedActor, FActorPathConsumer HandleRemovedActor, FActorConsumer HandleAddedActor) const
 {
-	DECLARE_SCOPE_CYCLE_COUNTER(TEXT("DiffWorld"), STAT_DiffWorld, STATGROUP_LevelSnapshots);
+	SCOPED_SNAPSHOT_CORE_TRACE(DiffWorld);
+	
 	if (!ensure(World && HandleMatchedActor.IsBound() && HandleRemovedActor.IsBound() && HandleAddedActor.IsBound()))
 	{
 		return;
 	}
+	UE_LOG(LogLevelSnapshots, Log, TEXT("Diffing snapshot %s in world %s"), *GetPathName(), *World->GetPathName());
+	ON_SCOPE_EXIT
+	{
+		UE_LOG(LogLevelSnapshots, Log, TEXT("Finished diffing snapshot"));
+	};
+
 
 	// Find actors that are not present in the snapshot
 	TSet<AActor*> AllActors;
-	const int32 NumActorsInWorld = Algo::Accumulate(World->GetLevels(), 0, [](int64 Size, const ULevel* Level){ return Size + Level->Actors.Num(); });
-	AllActors.Reserve(NumActorsInWorld);
-	for (ULevel* Level : World->GetLevels())
 	{
-		for (AActor* ActorInLevel : Level->Actors)
+		SCOPED_SNAPSHOT_CORE_TRACE(DiffWorld_FindAllActors);
+
+		const int32 NumActorsInWorld = Algo::Accumulate(World->GetLevels(), 0, [](int64 Size, const ULevel* Level){ return Size + Level->Actors.Num(); });
+		AllActors.Reserve(NumActorsInWorld);
+		for (ULevel* Level : World->GetLevels())
 		{
-			AllActors.Add(ActorInLevel);
-			
-			// Warning: ActorInLevel can be null, e.g. when an actor was just removed from the world (and still in undo buffer)
-			if (IsValid(ActorInLevel) && !SerializedData.HasMatchingSavedActor(ActorInLevel) && FSnapshotRestorability::ShouldConsiderNewActorForRemoval(ActorInLevel))
+			for (AActor* ActorInLevel : Level->Actors)
 			{
-				HandleAddedActor.Execute(ActorInLevel);
+				AllActors.Add(ActorInLevel);
+			
+				// Warning: ActorInLevel can be null, e.g. when an actor was just removed from the world (and still in undo buffer)
+				if (IsValid(ActorInLevel) && !SerializedData.HasMatchingSavedActor(ActorInLevel) && FSnapshotRestorability::ShouldConsiderNewActorForRemoval(ActorInLevel))
+				{
+					HandleAddedActor.Execute(ActorInLevel);
+				}
 			}
 		}
 	}
+	
+
 
 	// Try to find world actors and call appropriate callback
-	SerializedData.ForEachOriginalActor([&HandleMatchedActor, &HandleRemovedActor, &AllActors](const FSoftObjectPath& OriginalActorPath)
-    {
-		// TODO: we need to check whether the actor's class was blacklisted in the project settings
-		UObject* ResolvedActor = OriginalActorPath.ResolveObject();
-		// OriginalActorPath may still resolve to a live actor if it was just removed. We need to check the ULevel::Actors to see whether it was removed.
-		const bool bWasRemovedFromWorld = ResolvedActor == nullptr || !AllActors.Contains(Cast<AActor>(ResolvedActor));
-
-		// We do not need to call IsActorDesirableForCapture: it was already called when we took this snapshot
-		if (bWasRemovedFromWorld)
-		{
-			HandleRemovedActor.Execute(OriginalActorPath);
-		}
-		else
+	{
+		SCOPED_SNAPSHOT_CORE_TRACE(DiffWorld_IteratorAllActors);
+		
+		SerializedData.ForEachOriginalActor([&HandleMatchedActor, &HandleRemovedActor, &HandleAddedActor, &AllActors](const FSoftObjectPath& OriginalActorPath, const FActorSnapshotData& SavedData)
 		{
-			HandleMatchedActor.Execute(OriginalActorPath);
-		}
-    });
+			// TODO: we need to check whether the actor's class was blacklisted in the project settings
+			UObject* ResolvedActor = OriginalActorPath.ResolveObject();
+			// OriginalActorPath may still resolve to a live actor if it was just removed. We need to check the ULevel::Actors to see whether it was removed.
+			const bool bWasRemovedFromWorld = ResolvedActor == nullptr || !AllActors.Contains(Cast<AActor>(ResolvedActor));
+
+			// We do not need to call IsActorDesirableForCapture: it was already called when we took this snapshot
+			if (bWasRemovedFromWorld)
+			{
+				HandleRemovedActor.Execute(OriginalActorPath);
+				return;
+			}
+
+			UClass* ActorClass = SavedData.GetActorClass().TryLoadClass<AActor>();
+			if (!ActorClass)
+			{
+				UE_LOG(LogLevel, Warning, TEXT("Cannot find class %s. Saved actor %s will not be restored."), *SavedData.GetActorClass().ToString(), *OriginalActorPath.ToString());
+				return;
+			}
+
+			// Possible scenario: Right-click actor > Replace Selected Actors with; deletes the original and replaces it with new actor.
+			if (ResolvedActor->GetClass() != ActorClass)
+			{
+				HandleRemovedActor.Execute(OriginalActorPath);
+				HandleAddedActor.Execute(Cast<AActor>(ResolvedActor));
+			}
+			else
+			{
+				HandleMatchedActor.Execute(OriginalActorPath);
+			}
+		});
+	}
 }
 
 void ULevelSnapshot::SetSnapshotName(const FName& InSnapshotName)
diff --git a/Engine/Plugins/VirtualProduction/LevelSnapshots/Source/LevelSnapshots/Private/Data/WorldSnapshotData.cpp b/Engine/Plugins/VirtualProduction/LevelSnapshots/Source/LevelSnapshots/Private/Data/WorldSnapshotData.cpp
--- a/Engine/Plugins/VirtualProduction/LevelSnapshots/Source/LevelSnapshots/Private/Data/WorldSnapshotData.cpp
+++ b/Engine/Plugins/VirtualProduction/LevelSnapshots/Source/LevelSnapshots/Private/Data/WorldSnapshotData.cpp
@@ -6,7 +6,6 @@
 #include "Archive/ClassDefaults/ApplyClassDefaulDataArchive.h"
 #include "Archive/ClassDefaults/TakeClassDefaultObjectSnapshotArchive.h"
 #include "LevelSnapshotsLog.h"
-#include "LevelSnapshotsStats.h"
 #include "PropertySelectionMap.h"
 #include "Restorability/SnapshotRestorability.h"
 
@@ -182,8 +181,6 @@ void FWorldSnapshotData::OnDestroySnapshotWorld()
 
 void FWorldSnapshotData::SnapshotWorld(UWorld* World)
 {
-	DECLARE_SCOPE_CYCLE_COUNTER(TEXT("SnapshotWorld"), STAT_SnapshotWorld, STATGROUP_LevelSnapshots);
-
 	ClassDefaults.Empty();
 	ActorData.Empty();
 	SerializedNames.Empty();
@@ -214,7 +211,6 @@ void FWorldSnapshotData::SnapshotWorld(UWorld* World)
 
 void FWorldSnapshotData::ApplyToWorld(UWorld* WorldToApplyTo, UPackage* LocalisationSnapshotPackage, const FPropertySelectionMap& PropertiesToSerialize)
 {
-	DECLARE_SCOPE_CYCLE_COUNTER(TEXT("ApplySnapshotToWorld"), STAT_ApplySnapshotToWorld, STATGROUP_LevelSnapshots);
 	if (WorldToApplyTo == nullptr)
 	{
 		return;
@@ -260,11 +256,11 @@ int32 FWorldSnapshotData::GetNumSavedActors() const
 	return ActorData.Num();
 }
 
-void FWorldSnapshotData::ForEachOriginalActor(TFunction<void(const FSoftObjectPath& ActorPath)> HandleOriginalActorPath) const
+void FWorldSnapshotData::ForEachOriginalActor(TFunction<void(const FSoftObjectPath&, const FActorSnapshotData&)> HandleOriginalActorPath) const
 {
 	for (auto OriginalPathIt = ActorData.CreateConstIterator(); OriginalPathIt; ++OriginalPathIt)
 	{
-		HandleOriginalActorPath(OriginalPathIt->Key);
+		HandleOriginalActorPath(OriginalPathIt->Key, OriginalPathIt->Value);
 	}	
 }
 
@@ -608,6 +604,8 @@ UObject* FWorldSnapshotData::ResolveExternalReference(const FSoftObjectPath& Obj
 
 void FWorldSnapshotData::ApplyToWorld_HandleRemovingActors(UWorld* WorldToApplyTo, const FPropertySelectionMap& PropertiesToSerialize)
 {
+	SCOPED_SNAPSHOT_CORE_TRACE(ApplyToWorld_RemoveActors)
+	
 #if WITH_EDITOR
 	const TSet<TWeakObjectPtr<AActor>>& ActorsToDespawn = PropertiesToSerialize.GetNewActorsToDespawn();
 	const bool bShouldDespawnActors = ActorsToDespawn.Num() > 0;
@@ -654,8 +652,33 @@ void FWorldSnapshotData::ApplyToWorld_HandleRemovingActors(UWorld* WorldToApplyT
 #endif
 }
 
+namespace
+{
+	void DeleteActor(AActor* ActorToDespawn)
+	{
+		check(ActorToDespawn);
+
+		if ( IsValid(ActorToDespawn) )
+		{
+#if WITH_EDITOR
+			GEditor->SelectActor(ActorToDespawn, /*bSelect =*/true, /*bNotifyForActor =*/false, /*bSelectEvenIfHidden =*/true);
+			const bool bVerifyDeletionCanHappen = true;
+			const bool bWarnAboutReferences = false;
+			GEditor->edactDeleteSelected(ActorToDespawn->GetWorld(), bVerifyDeletionCanHappen, bWarnAboutReferences, bWarnAboutReferences);
+#else
+			ActorToDespawn->Destroy(true, true);
+#endif
+		}
+
+		const FName NewName = MakeUniqueObjectName(ActorToDespawn->GetWorld(), ActorToDespawn->GetClass());
+		ActorToDespawn->Rename(*NewName.ToString(), nullptr, REN_NonTransactional);
+	}
+}
+
 void FWorldSnapshotData::ApplyToWorld_HandleRecreatingActors(TSet<AActor*>& EvaluatedActors, UPackage* LocalisationSnapshotPackage, const FPropertySelectionMap& PropertiesToSerialize)
 {
+	SCOPED_SNAPSHOT_CORE_TRACE(ApplyToWorld_RecreateActors);
+	
 #if WITH_EDITOR
 	FScopedSlowTask RecreateActors(PropertiesToSerialize.GetDeletedActorsToRespawn().Num(), LOCTEXT("ApplyToWorld.RecreateActorsKey", "Re-creating actors"));
 	RecreateActors.MakeDialogDelayed(1.f, false);
@@ -671,20 +694,17 @@ void FWorldSnapshotData::ApplyToWorld_HandleRecreatingActors(TSet<AActor*>& Eval
 			continue;	
 		}
 
-		UClass* ActorClass = ActorSnapshot->GetActorClass().ResolveClass();
-		if (!ensure(ActorClass))
+		UClass* ActorClass = ActorSnapshot->GetActorClass().TryLoadClass<AActor>();
+		if (!ActorClass)
 		{
 			UE_LOG(LogLevelSnapshots, Warning, TEXT("Failed to resolve class '%s'. Was it removed?"), *ActorSnapshot->GetActorClass().ToString());
 			continue;
 		}
 		
-		// The actor may have just been removed and still exist in memory
-		if (UObject* StillExistingRemovedActor = OriginalRemovedActorPath.ResolveObject())
+		if (UObject* NameClash = OriginalRemovedActorPath.ResolveObject())
 		{
-			AActor* Actor = Cast<AActor>(StillExistingRemovedActor);
-			// We need to rename the trash object otherwise SpawnActor will have a name collision and assert.
-			const FName NewName = MakeUniqueObjectName(Actor->GetWorld(), ActorClass, *Actor->GetName().Append(TEXT("_TRASH")));
-			Actor->Rename(*NewName.ToString());
+			AActor* Actor = Cast<AActor>(NameClash);
+			DeleteActor(Actor);
 		}
 		
 		// Example: /Game/MapName.MapName:PersistentLevel.StaticMeshActor_42.StaticMeshComponent becomes /Game/MapName.MapName
@@ -715,7 +735,11 @@ void FWorldSnapshotData::ApplyToWorld_HandleRecreatingActors(TSet<AActor*>& Eval
 			SpawnParameters.Name = FName(*ActorName);
 			SpawnParameters.bNoFail = true;
 			SpawnParameters.Template = Cast<AActor>(GetClassDefault(ActorClass));
-			RecreatedActors.Add(OriginalRemovedActorPath, OwningLevelWorld->SpawnActor(ActorClass, nullptr, SpawnParameters));
+			SpawnParameters.NameMode = FActorSpawnParameters::ESpawnActorNameMode::Required_ErrorAndReturnNull;
+			if (AActor* RecreatedActor = OwningLevelWorld->SpawnActor(ActorClass, nullptr, SpawnParameters))
+			{
+				RecreatedActors.Add(OriginalRemovedActorPath, RecreatedActor);
+			}
 		}
 	}
 	
@@ -728,18 +752,24 @@ void FWorldSnapshotData::ApplyToWorld_HandleRecreatingActors(TSet<AActor*>& Eval
 		if (FActorSnapshotData* ActorSnapshot = ActorData.Find(OriginalRemovedActorPath))
 		{
 			AActor** RecreatedActor = RecreatedActors.Find(OriginalRemovedActorPath);
-			if (ensure(RecreatedActor))
+			if (RecreatedActor)
 			{
 				// Mark it, otherwise we'll serialize it again when we look for world actors matching the snapshot
 				EvaluatedActors.Add(*RecreatedActor);
 				ActorSnapshot->DeserializeIntoRecreatedEditorWorldActor(TempActorWorld.Get(), *RecreatedActor, *this, LocalisationSnapshotPackage, PropertiesToSerialize);
 			}
+			else
+			{
+				UE_LOG(LogLevelSnapshots, Error, TEXT("Failed to recreate actor %s"), *OriginalRemovedActorPath.ToString());
+			}
 		}
 	}
 }
 
 void FWorldSnapshotData::ApplyToWorld_HandleSerializingMatchingActors(TSet<AActor*>& EvaluatedActors, const TArray<FSoftObjectPath>& SelectedPaths, UPackage* LocalisationSnapshotPackage, const FPropertySelectionMap& PropertiesToSerialize)
 {
+	SCOPED_SNAPSHOT_CORE_TRACE(ApplyToWorld_SerializeMatchedActors);
+	
 #if WITH_EDITOR
 	FScopedSlowTask ExitingActorTask(SelectedPaths.Num(), LOCTEXT("ApplyToWorld.MatchingPropertiesKey", "Writing existing actors"));
 	ExitingActorTask.MakeDialogDelayed(1.f, true);
diff --git a/Engine/Plugins/VirtualProduction/LevelSnapshots/Source/LevelSnapshots/Private/LevelSnapshotsFunctionLibrary.cpp b/Engine/Plugins/VirtualProduction/LevelSnapshots/Source/LevelSnapshots/Private/LevelSnapshotsFunctionLibrary.cpp
--- a/Engine/Plugins/VirtualProduction/LevelSnapshots/Source/LevelSnapshots/Private/LevelSnapshotsFunctionLibrary.cpp
+++ b/Engine/Plugins/VirtualProduction/LevelSnapshots/Source/LevelSnapshots/Private/LevelSnapshotsFunctionLibrary.cpp
@@ -7,7 +7,6 @@
 #include "LevelSnapshot.h"
 #include "LevelSnapshotFilters.h"
 #include "LevelSnapshotsLog.h"
-#include "LevelSnapshotsStats.h"
 
 #include "EngineUtils.h"
 #include "CustomSerialization/CustomObjectSerializationWrapper.h"
@@ -98,8 +97,6 @@ ULevelSnapshot* ULevelSnapshotsFunctionLibrary::TakeLevelSnapshot(const UObject*
 
 ULevelSnapshot* ULevelSnapshotsFunctionLibrary::TakeLevelSnapshot_Internal(const UObject* WorldContextObject, const FName NewSnapshotName, UPackage* InPackage, const FString Description)
 {
-	DECLARE_SCOPE_CYCLE_COUNTER(TEXT("TakeLevelSnapshot"), STAT_TakeLevelSnapshot, STATGROUP_LevelSnapshots);
-	
 	UWorld* TargetWorld = nullptr;
 	if (WorldContextObject)
 	{
@@ -137,7 +134,6 @@ void ULevelSnapshotsFunctionLibrary::ApplyFilterToFindSelectedProperties(
 	bool bAllowUnchangedProperties,
     bool bAllowNonEditableProperties)
 {
-	DECLARE_SCOPE_CYCLE_COUNTER(TEXT("ApplyFilterToFindSelectedProperties"), STAT_ApplyFilterToFindSelectedProperties, STATGROUP_LevelSnapshots);
 	if (Filter == nullptr)
 	{
 		Filter = GetMutableDefault<UConstantFilter>();
diff --git a/Engine/Plugins/VirtualProduction/LevelSnapshots/Source/LevelSnapshots/Private/PropertySelectionMap.cpp b/Engine/Plugins/VirtualProduction/LevelSnapshots/Source/LevelSnapshots/Private/PropertySelectionMap.cpp
--- a/Engine/Plugins/VirtualProduction/LevelSnapshots/Source/LevelSnapshots/Private/PropertySelectionMap.cpp
+++ b/Engine/Plugins/VirtualProduction/LevelSnapshots/Source/LevelSnapshots/Private/PropertySelectionMap.cpp
@@ -2,7 +2,7 @@
 
 #include "PropertySelectionMap.h"
 
-#include "LevelSnapshotsStats.h"
+#include "LevelSnapshotsLog.h"
 
 #include "GameFramework/Actor.h"
 #include "UObject/UObjectHash.h"
@@ -68,8 +68,7 @@ void FPropertySelectionMap::Empty(bool bCanShrink)
 
 TArray<UObject*> FPropertySelectionMap::GetDirectSubobjectsWithProperties(UObject* Root) const
 {
-	// TODO post 4.27: Profile this and possibly build a TMap<UObject*, TArray<UObject*>> as the map is created
-	DECLARE_SCOPE_CYCLE_COUNTER(TEXT("GetDirectSubobjectsWithProperties"), STAT_GetDirectSubobjectsWithProperties, STATGROUP_LevelSnapshots);
+	SCOPED_SNAPSHOT_CORE_TRACE(GetDirectSubobjectsWithProperties);
 	
 	TArray<UObject*> Subobjects;
 	GetObjectsWithOuter(Root, Subobjects, true);
diff --git a/Engine/Plugins/VirtualProduction/LevelSnapshots/Source/LevelSnapshots/Private/Restorability/SnapshotRestorability.cpp b/Engine/Plugins/VirtualProduction/LevelSnapshots/Source/LevelSnapshots/Private/Restorability/SnapshotRestorability.cpp
--- a/Engine/Plugins/VirtualProduction/LevelSnapshots/Source/LevelSnapshots/Private/Restorability/SnapshotRestorability.cpp
+++ b/Engine/Plugins/VirtualProduction/LevelSnapshots/Source/LevelSnapshots/Private/Restorability/SnapshotRestorability.cpp
@@ -2,8 +2,8 @@
 
 #include "Restorability/SnapshotRestorability.h"
 
+#include "LevelSnapshotsLog.h"
 #include "LevelSnapshotsModule.h"
-#include "LevelSnapshotsStats.h"
 
 #include "Components/ActorComponent.h"
 #include "Components/BillboardComponent.h"
@@ -76,7 +76,7 @@ namespace
 
 bool FSnapshotRestorability::IsActorDesirableForCapture(const AActor* Actor)
 {
-	DECLARE_SCOPE_CYCLE_COUNTER(TEXT("IsActorDesirableForCapture"), STAT_IsActorDesirableForCapture, STATGROUP_LevelSnapshots);
+	SCOPED_SNAPSHOT_CORE_TRACE(IsActorDesirableForCapture);
 	
 	bool bSomebodyAllowed = false;
 	const TArray<TSharedRef<ISnapshotRestorabilityOverrider>>& Overrides = FLevelSnapshotsModule::GetInternalModuleInstance().GetOverrides(); 
@@ -112,7 +112,7 @@ bool FSnapshotRestorability::IsActorDesirableForCapture(const AActor* Actor)
 
 bool FSnapshotRestorability::IsComponentDesirableForCapture(const UActorComponent* Component)
 {
-	DECLARE_SCOPE_CYCLE_COUNTER(TEXT("IsComponentDesirableForCapture"), STAT_IsComponentDesirableForCapture, STATGROUP_LevelSnapshots);
+	SCOPED_SNAPSHOT_CORE_TRACE(IsComponentDesirableForCapture);
 	
 	bool bSomebodyAllowed = false;
 	const TArray<TSharedRef<ISnapshotRestorabilityOverrider>>& Overrides = FLevelSnapshotsModule::GetInternalModuleInstance().GetOverrides(); 
@@ -140,11 +140,13 @@ bool FSnapshotRestorability::IsComponentDesirableForCapture(const UActorComponen
 
 bool FSnapshotRestorability::IsPropertyBlacklistedForCapture(const FProperty* Property)
 {
+	SCOPED_SNAPSHOT_CORE_TRACE(IsPropertyBlacklisted);
 	return FLevelSnapshotsModule::GetInternalModuleInstance().IsPropertyBlacklisted(Property);
 }
 
 bool FSnapshotRestorability::IsPropertyWhitelistedForCapture(const FProperty* Property)
 {
+	SCOPED_SNAPSHOT_CORE_TRACE(IsPropertyWhitelisted);
 	return FLevelSnapshotsModule::GetInternalModuleInstance().IsPropertyWhitelisted(Property);
 }
 
@@ -155,7 +157,7 @@ bool FSnapshotRestorability::ShouldConsiderNewActorForRemoval(const AActor* Acto
 
 bool FSnapshotRestorability::IsRestorableProperty(const FProperty* LeafProperty)
 {
-	DECLARE_SCOPE_CYCLE_COUNTER(TEXT("IsRestorableProperty"), STAT_IsRestorableProperty, STATGROUP_LevelSnapshots);
+	SCOPED_SNAPSHOT_CORE_TRACE(IsRestorableProperty);
 	
 	// Subobjects are currently not supported
 	const uint64 InstancedFlags = CPF_InstancedReference | CPF_ContainsInstancedReference | CPF_PersistentInstance;
diff --git a/Engine/Plugins/VirtualProduction/LevelSnapshots/Source/LevelSnapshots/Private/Tests/LevelSnapshotsLog.cpp b/Engine/Plugins/VirtualProduction/LevelSnapshots/Source/LevelSnapshots/Private/Tests/LevelSnapshotsLog.cpp
new file mode 100644
--- /dev/null
+++ b/Engine/Plugins/VirtualProduction/LevelSnapshots/Source/LevelSnapshots/Private/Tests/LevelSnapshotsLog.cpp
@@ -0,0 +1,5 @@
+﻿// Copyright Epic Games, Inc. All Rights Reserved.
+
+#include "LevelSnapshotsLog.h"
+
+UE_TRACE_CHANNEL_DEFINE(LevelSnapshotsChannel);
\ No newline at end of file
diff --git a/Engine/Plugins/VirtualProduction/LevelSnapshots/Source/LevelSnapshots/Public/Data/WorldSnapshotData.h b/Engine/Plugins/VirtualProduction/LevelSnapshots/Source/LevelSnapshots/Public/Data/WorldSnapshotData.h
--- a/Engine/Plugins/VirtualProduction/LevelSnapshots/Source/LevelSnapshots/Public/Data/WorldSnapshotData.h
+++ b/Engine/Plugins/VirtualProduction/LevelSnapshots/Source/LevelSnapshots/Public/Data/WorldSnapshotData.h
@@ -36,7 +36,7 @@ struct LEVELSNAPSHOTS_API FWorldSnapshotData
 	void ApplyToWorld(UWorld* WorldToApplyTo, UPackage* LocalisationSnapshotPackage, const FPropertySelectionMap& PropertiesToSerialize);
 
 	int32 GetNumSavedActors() const;
-	void ForEachOriginalActor(TFunction<void(const FSoftObjectPath& ActorPath)> HandleOriginalActorPath) const;
+	void ForEachOriginalActor(TFunction<void(const FSoftObjectPath& ActorPath, const FActorSnapshotData& SavedData)> HandleOriginalActorPath) const;
 	bool HasMatchingSavedActor(const FSoftObjectPath& OriginalObjectPath) const;
 
 	/* Given a path to a world actor, gets the equivalent allocated actor. */
diff --git a/Engine/Plugins/VirtualProduction/LevelSnapshots/Source/LevelSnapshots/Public/LevelSnapshotsLog.h b/Engine/Plugins/VirtualProduction/LevelSnapshots/Source/LevelSnapshots/Public/LevelSnapshotsLog.h
--- a/Engine/Plugins/VirtualProduction/LevelSnapshots/Source/LevelSnapshots/Public/LevelSnapshotsLog.h
+++ b/Engine/Plugins/VirtualProduction/LevelSnapshots/Source/LevelSnapshots/Public/LevelSnapshotsLog.h
@@ -3,5 +3,18 @@
 #pragma once
 
 #include "CoreMinimal.h"
+#include "ProfilingDebugging/CpuProfilerTrace.h"
 
-DECLARE_LOG_CATEGORY_CLASS(LogLevelSnapshots,   Log, All);
+#if UE_TRACE_ENABLED && !UE_BUILD_SHIPPING
+UE_TRACE_CHANNEL_EXTERN(LevelSnapshotsChannel, LEVELSNAPSHOTS_API);
+#define SCOPED_SNAPSHOT_TRACE(TraceName) TRACE_CPUPROFILER_EVENT_SCOPE_ON_CHANNEL(TraceName, LevelSnapshotsChannel)
+#else
+#define SCOPED_SNAPSHOT_TRACE(TraceName)
+#endif
+
+// Use inside 'implementation' modules, like this one or filters
+#define SCOPED_SNAPSHOT_CORE_TRACE(TraceName) SCOPED_SNAPSHOT_TRACE(SnapshotCore_##TraceName)
+// Use inside editor modules
+#define SCOPED_SNAPSHOT_EDITOR_TRACE(TraceName) SCOPED_SNAPSHOT_TRACE(SnapshotEditor_##TraceName)
+
+DECLARE_LOG_CATEGORY_CLASS(LogLevelSnapshots, Log, All);
diff --git a/Engine/Plugins/VirtualProduction/LevelSnapshots/Source/LevelSnapshots/Public/LevelSnapshotsStats.h b/Engine/Plugins/VirtualProduction/LevelSnapshots/Source/LevelSnapshots/Public/LevelSnapshotsStats.h
deleted file mode 100644
diff --git a/Engine/Plugins/VirtualProduction/LevelSnapshots/Source/LevelSnapshotsEditor/LevelSnapshotsEditor.build.cs b/Engine/Plugins/VirtualProduction/LevelSnapshots/Source/LevelSnapshotsEditor/LevelSnapshotsEditor.build.cs
--- a/Engine/Plugins/VirtualProduction/LevelSnapshots/Source/LevelSnapshotsEditor/LevelSnapshotsEditor.build.cs
+++ b/Engine/Plugins/VirtualProduction/LevelSnapshots/Source/LevelSnapshotsEditor/LevelSnapshotsEditor.build.cs
@@ -50,6 +50,7 @@ public class LevelSnapshotsEditor : ModuleRules
 				"EditorFramework",
 				"EditorStyle",
 				"EditorWidgets",
+				"GameProjectGeneration",
 				"InputCore",
 				"Kismet",
 				"LevelSnapshots",
diff --git a/Engine/Plugins/VirtualProduction/LevelSnapshots/Source/LevelSnapshotsEditor/Private/Data/FilterListData.cpp b/Engine/Plugins/VirtualProduction/LevelSnapshots/Source/LevelSnapshotsEditor/Private/Data/FilterListData.cpp
--- a/Engine/Plugins/VirtualProduction/LevelSnapshots/Source/LevelSnapshotsEditor/Private/Data/FilterListData.cpp
+++ b/Engine/Plugins/VirtualProduction/LevelSnapshots/Source/LevelSnapshotsEditor/Private/Data/FilterListData.cpp
@@ -6,7 +6,6 @@
 #include "LevelSnapshotFilters.h"
 #include "LevelSnapshotsFunctionLibrary.h"
 #include "LevelSnapshotsLog.h"
-#include "LevelSnapshotsStats.h"
 
 #include "GameFramework/Actor.h"
 #include "Misc/ScopedSlowTask.h"
@@ -15,6 +14,8 @@
 
 void FFilterListData::UpdateFilteredList(UWorld* World, ULevelSnapshot* FromSnapshot, ULevelSnapshotFilter* FilterToApply)
 {
+	SCOPED_SNAPSHOT_EDITOR_TRACE(UpdateFilteredList);
+	
 	// We only track progress of HandleActorExistsInWorldAndSnapshot because the other two functions are relatively fast in comparison: deserialisation takes much longer.
 	const int32 ExpectedAmountOfWork = FromSnapshot->GetNumSavedActors();
 	FScopedSlowTask DiffDeserializedActors(ExpectedAmountOfWork, LOCTEXT("DiffingActorsKey", "Diffing actors"));
@@ -38,8 +39,6 @@ void FFilterListData::UpdateFilteredList(UWorld* World, ULevelSnapshot* FromSnap
 
 void FFilterListData::ApplyFilterToFindSelectedProperties(AActor* WorldActor, ULevelSnapshotFilter* FilterToApply)
 {
-	DECLARE_SCOPE_CYCLE_COUNTER(TEXT("ApplyFilterToFindSelectedProperties"), STAT_ApplyFilterToFindSelectedProperties, STATGROUP_LevelSnapshots);
-	
 	const FPropertySelection* AllowedSelectedProperties = ModifiedActorsSelectedProperties_AllowedByFilter.GetSelectedProperties(WorldActor);
 	const FPropertySelection* DisallowedSelectedProperties = ModifiedActorsSelectedProperties_AllowedByFilter.GetSelectedProperties(WorldActor);
 	
diff --git a/Engine/Plugins/VirtualProduction/LevelSnapshots/Source/LevelSnapshotsEditor/Private/Data/FilteredResults.cpp b/Engine/Plugins/VirtualProduction/LevelSnapshots/Source/LevelSnapshotsEditor/Private/Data/FilteredResults.cpp
--- a/Engine/Plugins/VirtualProduction/LevelSnapshots/Source/LevelSnapshotsEditor/Private/Data/FilteredResults.cpp
+++ b/Engine/Plugins/VirtualProduction/LevelSnapshots/Source/LevelSnapshotsEditor/Private/Data/FilteredResults.cpp
@@ -5,7 +5,6 @@
 #include "LevelSnapshot.h"
 #include "LevelSnapshotFilters.h"
 #include "LevelSnapshotsFunctionLibrary.h"
-#include "LevelSnapshotsStats.h"
 
 #include "EngineUtils.h"
 #include "Stats/StatsMisc.h"
@@ -29,7 +28,6 @@ void UFilteredResults::SetUserFilters(ULevelSnapshotFilter* InUserFilters)
 
 void UFilteredResults::UpdateFilteredResults(UWorld* SelectedWorld)
 {
-	DECLARE_SCOPE_CYCLE_COUNTER(TEXT("UpdateFilteredResults"), STAT_UpdateFilteredResults, STATGROUP_LevelSnapshots);
 	if (!ensure(UserSelectedSnapshot.IsValid()) || !ensure(UserFilters.IsValid()) || !ensure(SelectedWorld))
 	{
 		return;
diff --git a/Engine/Plugins/VirtualProduction/LevelSnapshots/Source/LevelSnapshotsEditor/Private/Data/Filters/ConjunctionFilter.h b/Engine/Plugins/VirtualProduction/LevelSnapshots/Source/LevelSnapshotsEditor/Private/Data/Filters/ConjunctionFilter.h
--- a/Engine/Plugins/VirtualProduction/LevelSnapshots/Source/LevelSnapshotsEditor/Private/Data/Filters/ConjunctionFilter.h
+++ b/Engine/Plugins/VirtualProduction/LevelSnapshots/Source/LevelSnapshotsEditor/Private/Data/Filters/ConjunctionFilter.h
@@ -49,7 +49,7 @@ public:
 	
 private:
 
-	UPROPERTY()
+	UPROPERTY(Instanced)
 	TArray<UNegatableFilter*> Children;
 
 	UPROPERTY()
diff --git a/Engine/Plugins/VirtualProduction/LevelSnapshots/Source/LevelSnapshotsEditor/Private/Data/Filters/LevelSnapshotsFilterPreset.cpp b/Engine/Plugins/VirtualProduction/LevelSnapshots/Source/LevelSnapshotsEditor/Private/Data/Filters/LevelSnapshotsFilterPreset.cpp
--- a/Engine/Plugins/VirtualProduction/LevelSnapshots/Source/LevelSnapshotsEditor/Private/Data/Filters/LevelSnapshotsFilterPreset.cpp
+++ b/Engine/Plugins/VirtualProduction/LevelSnapshots/Source/LevelSnapshotsEditor/Private/Data/Filters/LevelSnapshotsFilterPreset.cpp
@@ -66,12 +66,35 @@ void ULevelSnapshotsFilterPreset::MarkTransactional()
 	{
 		Child->MarkTransactional();
 	}
+
+	OnObjectTransactedHandle = FCoreUObjectDelegates::OnObjectTransacted.AddLambda([this](UObject* ModifiedObject, const FTransactionObjectEvent& TransactionInfo)
+	{
+		if (!ModifiedObject || bHasJustCreatedNewChild)
+		{
+			return;
+		}
+
+		const bool bChangedChildArray = ModifiedObject == this && TransactionInfo.GetChangedProperties().Contains(GET_MEMBER_NAME_CHECKED(ULevelSnapshotsFilterPreset, Children));
+		if (bChangedChildArray)
+		{
+			OnFilterModified.Broadcast(EFilterChangeType::RowChildFilterAddedOrRemoved);
+			return;
+		}
+		
+		const bool bChangedChildInstance = ModifiedObject != this && ModifiedObject->IsIn(this);
+		const bool bConjunctionWasAddedOrRemoved = TransactionInfo.HasPendingKillChange();
+		if (bChangedChildInstance && !bConjunctionWasAddedOrRemoved)
+		{
+			const bool bConjunctionChildrenWereChanged = TransactionInfo.GetChangedProperties().Contains(UConjunctionFilter::GetChildrenMemberName());
+			const bool bModifiedFilterHierarchy = Cast<UConjunctionFilter>(ModifiedObject) && (bConjunctionWasAddedOrRemoved || bConjunctionChildrenWereChanged);
+			OnFilterModified.Broadcast(bModifiedFilterHierarchy ? EFilterChangeType::RowChildFilterAddedOrRemoved : EFilterChangeType::FilterPropertyModified);
+		}
+	});
 }
 
 UConjunctionFilter* ULevelSnapshotsFilterPreset::CreateChild()
 {
 	UConjunctionFilter* Child;
-
 	{
 		FScopedTransaction Transaction(FText::FromString("Add filter row"));
 		Modify();
@@ -79,11 +102,10 @@ UConjunctionFilter* ULevelSnapshotsFilterPreset::CreateChild()
 		Child = NewObject<UConjunctionFilter>(this, UConjunctionFilter::StaticClass(), NAME_None, RF_Transactional);
 		Children.Add(Child);
 
-		OnFilterModified.Broadcast(EFilterChangeType::RowAdded);
-		
+		OnFilterModified.Broadcast(EFilterChangeType::BlankRowAdded);
+		// We broadcast OnFilterModified already. Avoid OnObjectTransacted triggering it again.
 		bHasJustCreatedNewChild = true;
 	}
-	// OnObjectTransacted will be triggered twice when a ConjunctionFilter is created, this scope ignores the second call
 	bHasJustCreatedNewChild = false;
 	
 	
@@ -108,23 +130,6 @@ const TArray<UConjunctionFilter*>& ULevelSnapshotsFilterPreset::GetChildren() co
 	return Children;
 }
 
-void ULevelSnapshotsFilterPreset::PostInitProperties()
-{
-	Super::PostInitProperties();
-
-	if (!HasAnyFlags(RF_ClassDefaultObject | RF_ArchetypeObject))
-	{
-		OnObjectTransactedHandle = FCoreUObjectDelegates::OnObjectTransacted.AddLambda([this](UObject* ModifiedObject, const FTransactionObjectEvent& TransactionInfo)
-		{	
-			if (IsValid(ModifiedObject) && ModifiedObject->IsIn(this) && ModifiedObject != this && !bHasJustCreatedNewChild)
-			{
-				const bool bModifiedChildren = Cast<UConjunctionFilter>(ModifiedObject) && TransactionInfo.GetChangedProperties().Contains(UConjunctionFilter::GetChildrenMemberName());
-				OnFilterModified.Broadcast(bModifiedChildren ? EFilterChangeType::RowChildFilterAddedOrRemoved : EFilterChangeType::FilterPropertyModified);
-			}
-		});
-	}
-}
-
 void ULevelSnapshotsFilterPreset::BeginDestroy()
 {
 	if (OnObjectTransactedHandle.IsValid())
diff --git a/Engine/Plugins/VirtualProduction/LevelSnapshots/Source/LevelSnapshotsEditor/Private/Data/Filters/LevelSnapshotsFilterPreset.h b/Engine/Plugins/VirtualProduction/LevelSnapshots/Source/LevelSnapshotsEditor/Private/Data/Filters/LevelSnapshotsFilterPreset.h
--- a/Engine/Plugins/VirtualProduction/LevelSnapshots/Source/LevelSnapshotsEditor/Private/Data/Filters/LevelSnapshotsFilterPreset.h
+++ b/Engine/Plugins/VirtualProduction/LevelSnapshots/Source/LevelSnapshotsEditor/Private/Data/Filters/LevelSnapshotsFilterPreset.h
@@ -12,7 +12,7 @@ class UNegatableFilter;
 enum class EFilterChangeType : uint8
 {
 	/* A blank row was added */
-	RowAdded,
+	BlankRowAdded,
 	/* A row was removed */
 	RowRemoved,
 	/* A filter was added or removed to a row */
@@ -42,7 +42,6 @@ public:
 	const TArray<UConjunctionFilter*>& GetChildren() const;
 
 	//~ Begin UObject Interface
-	virtual void PostInitProperties() override;
 	virtual void BeginDestroy() override;
 	//~ End UObject Interface
 	
@@ -65,7 +64,7 @@ public:
 
 private:
 	
-	UPROPERTY()
+	UPROPERTY(Instanced)
 	TArray<UConjunctionFilter*> Children;
 
 	bool bHasJustCreatedNewChild = false;
diff --git a/Engine/Plugins/VirtualProduction/LevelSnapshots/Source/LevelSnapshotsEditor/Private/Data/Filters/NegatableFilter.h b/Engine/Plugins/VirtualProduction/LevelSnapshots/Source/LevelSnapshotsEditor/Private/Data/Filters/NegatableFilter.h
--- a/Engine/Plugins/VirtualProduction/LevelSnapshots/Source/LevelSnapshotsEditor/Private/Data/Filters/NegatableFilter.h
+++ b/Engine/Plugins/VirtualProduction/LevelSnapshots/Source/LevelSnapshotsEditor/Private/Data/Filters/NegatableFilter.h
@@ -64,6 +64,6 @@ private:
 	UPROPERTY(EditAnywhere, Category = "Filter")
 	bool bIgnoreFilter = false;
 	
-	UPROPERTY()
+	UPROPERTY(Instanced)
 	ULevelSnapshotFilter* ChildFilter;
 };
\ No newline at end of file
diff --git a/Engine/Plugins/VirtualProduction/LevelSnapshots/Source/LevelSnapshotsEditor/Private/Data/LevelSnapshotsEditorData.cpp b/Engine/Plugins/VirtualProduction/LevelSnapshots/Source/LevelSnapshotsEditor/Private/Data/LevelSnapshotsEditorData.cpp
--- a/Engine/Plugins/VirtualProduction/LevelSnapshots/Source/LevelSnapshotsEditor/Private/Data/LevelSnapshotsEditorData.cpp
+++ b/Engine/Plugins/VirtualProduction/LevelSnapshots/Source/LevelSnapshotsEditor/Private/Data/LevelSnapshotsEditorData.cpp
@@ -1,16 +1,15 @@
 // Copyright Epic Games, Inc. All Rights Reserved.
 
 #include "LevelSnapshotsEditorData.h"
-#include "SLevelSnapshotsEditorInput.h"
-
-#include "Editor.h"
-#include "Engine/World.h"
-#include "UObject/UObjectGlobals.h"
 
 #include "FavoriteFilterContainer.h"
 #include "FilterLoader.h"
 #include "FilteredResults.h"
-#include "SnapshotRestorability.h"
+#include "LevelSnapshotsLog.h"
+
+#include "Editor.h"
+#include "Engine/World.h"
+#include "UObject/UObjectGlobals.h"
 
 ULevelSnapshotsEditorData::ULevelSnapshotsEditorData(const FObjectInitializer& ObjectInitializer)
 {
@@ -22,7 +21,6 @@ ULevelSnapshotsEditorData::ULevelSnapshotsEditorData(const FObjectInitializer& O
 		this,
 		TEXT("UserDefinedFilters")
 		);
-	UserDefinedFilters->SetFlags(RF_Transactional);
 
 	TrackedFilterModifiedHandle = UserDefinedFilters->OnFilterModified.AddUObject(this, &ULevelSnapshotsEditorData::HandleFilterChange);
 	
@@ -53,22 +51,40 @@ ULevelSnapshotsEditorData::ULevelSnapshotsEditorData(const FObjectInitializer& O
         TEXT("FilterResults")
         );
 	FilterResults->SetUserFilters(UserDefinedFilters);
+}
 
-	OnWorldCleanup = FWorldDelegates::OnWorldCleanup.AddLambda([this](UWorld* World, bool bSessionEnded, bool bCleanupResources)
-    {
-        ClearActiveSnapshot();
-    });
+void ULevelSnapshotsEditorData::PostInitProperties()
+{
+	Super::PostInitProperties();
 
-	// Dirty filter state when an actor that is desirable for capture is modified.
-	OnObjectsEdited = FCoreUObjectDelegates::OnObjectModified.AddUObject(this, &ULevelSnapshotsEditorData::HandleWorldActorsEdited);
+	// Class default should not register to global events
+	if (!HasAnyFlags(RF_ClassDefaultObject))
+	{
+		UserDefinedFilters->MarkTransactional();
+		OnObjectsEdited = FCoreUObjectDelegates::OnObjectModified.AddUObject(this, &ULevelSnapshotsEditorData::HandleWorldActorsEdited);
+		
+		OnWorldCleanup = FWorldDelegates::OnWorldCleanup.AddLambda([this](UWorld* World, bool bSessionEnded, bool bCleanupResources)
+		{
+			ClearActiveSnapshot();
+		});
+	}
 }
 
 void ULevelSnapshotsEditorData::BeginDestroy()
 {
 	Super::BeginDestroy();
-	
-	FWorldDelegates::OnWorldCleanup.Remove(OnWorldCleanup);
-	FCoreUObjectDelegates::OnObjectModified.Remove(OnObjectsEdited);
+
+	if (!HasAnyFlags(RF_ClassDefaultObject))
+	{
+		FWorldDelegates::OnWorldCleanup.Remove(OnWorldCleanup);
+		FCoreUObjectDelegates::OnObjectModified.Remove(OnObjectsEdited);
+		
+		OnWorldCleanup.Reset();
+		OnObjectsEdited.Reset();
+	}
+
+	check(!OnWorldCleanup.IsValid());
+	check(!OnObjectsEdited.IsValid());
 }
 
 void ULevelSnapshotsEditorData::CleanupAfterEditorClose()
@@ -85,8 +101,9 @@ void ULevelSnapshotsEditorData::CleanupAfterEditorClose()
 
 void ULevelSnapshotsEditorData::SetActiveSnapshot(const TOptional<ULevelSnapshot*>& NewActiveSnapshot)
 {
+	SCOPED_SNAPSHOT_EDITOR_TRACE(SetActiveSnapshot);
+	
 	ActiveSnapshot = NewActiveSnapshot.Get(nullptr) ? TStrongObjectPtr<ULevelSnapshot>(NewActiveSnapshot.GetValue()) : TOptional<TStrongObjectPtr<ULevelSnapshot>>();
-
 	FilterResults->SetActiveLevelSnapshot(NewActiveSnapshot.Get(nullptr));
 	OnActiveSnapshotChanged.Broadcast(GetActiveSnapshot());
 }
@@ -105,7 +122,9 @@ TOptional<ULevelSnapshot*> ULevelSnapshotsEditorData::GetActiveSnapshot() const
 
 UWorld* ULevelSnapshotsEditorData::GetEditorWorld()
 {
-	if (GEditor)
+	// If this function is called very early during startup, the initial editor GWorld may not have been created yet!
+	const bool bIsEngineInitialised = GEditor && GEditor->GetWorldContexts().Num() > 0;
+	if (bIsEngineInitialised)
 	{
 		if (UWorld* World = GEditor->GetEditorWorldContext().World())
 		{
@@ -153,7 +172,8 @@ UFilteredResults* ULevelSnapshotsEditorData::GetFilterResults() const
 
 void ULevelSnapshotsEditorData::HandleFilterChange(EFilterChangeType FilterChangeType)
 {
-	if (FilterChangeType != EFilterChangeType::RowAdded)
+	// Blank rows do not functionally change the filter
+	if (FilterChangeType != EFilterChangeType::BlankRowAdded)
 	{
 		SetIsFilterDirty(true);
 	}
diff --git a/Engine/Plugins/VirtualProduction/LevelSnapshots/Source/LevelSnapshotsEditor/Private/Data/LevelSnapshotsEditorData.h b/Engine/Plugins/VirtualProduction/LevelSnapshots/Source/LevelSnapshotsEditor/Private/Data/LevelSnapshotsEditorData.h
--- a/Engine/Plugins/VirtualProduction/LevelSnapshots/Source/LevelSnapshotsEditor/Private/Data/LevelSnapshotsEditorData.h
+++ b/Engine/Plugins/VirtualProduction/LevelSnapshots/Source/LevelSnapshotsEditor/Private/Data/LevelSnapshotsEditorData.h
@@ -26,6 +26,7 @@ public:
 	ULevelSnapshotsEditorData(const FObjectInitializer& ObjectInitializer);
 
 	//~ Begin UObject Interface
+	virtual void PostInitProperties() override;
 	virtual void BeginDestroy() override;
 	//~ End UObject Interface
 	
diff --git a/Engine/Plugins/VirtualProduction/LevelSnapshots/Source/LevelSnapshotsEditor/Private/LevelSnapshotsEditorFunctionLibrary.cpp b/Engine/Plugins/VirtualProduction/LevelSnapshots/Source/LevelSnapshotsEditor/Private/LevelSnapshotsEditorFunctionLibrary.cpp
--- a/Engine/Plugins/VirtualProduction/LevelSnapshots/Source/LevelSnapshotsEditor/Private/LevelSnapshotsEditorFunctionLibrary.cpp
+++ b/Engine/Plugins/VirtualProduction/LevelSnapshots/Source/LevelSnapshotsEditor/Private/LevelSnapshotsEditorFunctionLibrary.cpp
@@ -4,27 +4,16 @@
 
 #include "LevelSnapshotsLog.h"
 
-#include "AssetRegistry/AssetRegistryModule.h"
-#include "AssetToolsModule.h"
-#include "AssetViewUtils.h"
-#include "IAssetTools.h"
 #include "Editor.h"
 #include "Engine/World.h"
 #include "LevelSnapshot.h"
-#include "LevelSnapshotsEditorStyle.h"
-#include "Framework/Notifications/NotificationManager.h"
 #include "GameplayMediaEncoder/Private/GameplayMediaEncoderCommon.h"
 #include "Logging/MessageLog.h"
-#include "Misc/MessageDialog.h"
 #include "UObject/MetaData.h"
-#include "UObject/Package.h"
-#include "Widgets/Notifications/SNotificationList.h"
+#include "Util/TakeSnapshotUtil.h"
 
-#define LOCTEXT_NAMESPACE "LevelSnapshotEditorLibrary"
-
-ULevelSnapshot* ULevelSnapshotsEditorFunctionLibrary::TakeLevelSnapshotAndSaveToDisk(
-	const UObject* WorldContextObject, const FString& FileName, const FString& FolderPath, const FString& Description, const bool bShouldCreateUniqueFileName)
-{
+ULevelSnapshot* ULevelSnapshotsEditorFunctionLibrary::TakeLevelSnapshotAndSaveToDisk(const UObject* WorldContextObject, const FString& FileName, const FString& FolderPath, const FString& Description, const bool bShouldCreateUniqueFileName)
+{	
 	UWorld* TargetWorld = nullptr;
 	if (WorldContextObject)
 	{
@@ -37,92 +26,7 @@ ULevelSnapshot* ULevelSnapshotsEditorFunctionLibrary::TakeLevelSnapshotAndSaveTo
 		return nullptr;
 	}
 
-	FString AssetName = FileName;
-
-	FString BasePackageName = FPaths::Combine(FolderPath, FileName);
-
-	BasePackageName.RemoveFromStart(TEXT("/"));
-	BasePackageName.RemoveFromStart(TEXT("Content/"));
-	BasePackageName.StartsWith(TEXT("Game/")) == true ? BasePackageName.InsertAt(0, TEXT("/")) : BasePackageName.InsertAt(0, TEXT("/Game/"));
-
-	FString PackageName = BasePackageName;
-
-	if (bShouldCreateUniqueFileName)
-	{
-		IAssetTools& AssetTools = FModuleManager::Get().LoadModuleChecked<FAssetToolsModule>("AssetTools").Get();
-		
-	AssetTools.CreateUniqueAssetName(BasePackageName, TEXT(""), PackageName, AssetName);
-	}
-
-	const FString& PackageFileName = FPackageName::LongPackageNameToFilename(PackageName, FPackageName::GetAssetPackageExtension());
-	const bool bPathIsValid = FPaths::ValidatePath(PackageFileName);
-
-	bool bProceedWithSave = bPathIsValid;
-
-	if (bProceedWithSave && !bShouldCreateUniqueFileName)
-	{
-		FAssetRegistryModule& AssetRegistryModule = FModuleManager::LoadModuleChecked<FAssetRegistryModule>("AssetRegistry");
-		FAssetData ExistingAsset = AssetRegistryModule.Get().GetAssetByObjectPath(FName(*(BasePackageName + "." + FileName)));
-		if ( ExistingAsset.IsValid() && ExistingAsset.AssetClass == *ULevelSnapshot::StaticClass()->GetName() )
-		{
-			EAppReturnType::Type ShouldReplace = FMessageDialog::Open( EAppMsgType::YesNo, FText::Format(LOCTEXT("ReplaceAssetMessage", "{0} already exists. Do you want to replace it?"), FText::FromString(FileName)) );
-			bProceedWithSave = (ShouldReplace == EAppReturnType::Yes);
-		}
-	}
-
-	if (bProceedWithSave)
-	{
-		// Show notification
-		FNotificationInfo Notification(NSLOCTEXT("LevelSnapshots", "NotificationFormatText_CreatingSnapshot", "Creating Level Snapshot"));
-		Notification.Image = FLevelSnapshotsEditorStyle::GetBrush(TEXT("LevelSnapshots.ToolbarButton"));
-		Notification.bUseThrobber = true;
-		Notification.bUseSuccessFailIcons = true;
-		Notification.ExpireDuration = 2.f;
-		Notification.bFireAndForget = false;
-		
-		TSharedPtr<SNotificationItem> NotificationItem = FSlateNotificationManager::Get().AddNotification(Notification);
-		NotificationItem->SetCompletionState(SNotificationItem::CS_Pending);
-		ON_SCOPE_EXIT
-		{
-			NotificationItem->ExpireAndFadeout();
-		};
-
-	UPackage* SavePackage = CreatePackage(*PackageName);
-	const EObjectFlags AssetFlags = RF_Public | RF_Standalone;
-	ULevelSnapshot* SnapshotAsset = NewObject<ULevelSnapshot>(SavePackage, *AssetName, AssetFlags, nullptr);
-
-		check(SavePackage && SnapshotAsset);
-		
-		SnapshotAsset->SetSnapshotName(*FileName);
-		SnapshotAsset->SetSnapshotDescription(Description);
-
-		SnapshotAsset->SnapshotWorld(TargetWorld);
-		SnapshotAsset->MarkPackageDirty();
-
-		FAssetRegistryModule::AssetCreated(SnapshotAsset);
-		GenerateThumbnailForSnapshotAsset(SnapshotAsset);
-
-		// Notify the user of the outcome
-		if (bPathIsValid && UPackage::SavePackage(SavePackage, SnapshotAsset, RF_Public | RF_Standalone, *PackageFileName))
-		{
-			// If successful
-			NotificationItem->SetText(
-				FText::Format(
-					NSLOCTEXT("LevelSnapshots", "NotificationFormatText_CreateSnapshotSuccess", "Successfully created Level Snapshot \"{0}\""), FText::FromString(FileName)));
-			NotificationItem->SetCompletionState(SNotificationItem::CS_Success);
-	}
-		else
-		{
-			NotificationItem->SetText(
-				FText::Format(
-					NSLOCTEXT("LevelSnapshots", "NotificationFormatText_CreateSnapshotSuccess", "Failed to create Level Snapshot \"{0}\". Check the file name."), FText::FromString(FileName)));
-			NotificationItem->SetCompletionState(SNotificationItem::CS_Fail);
-		}
-
-	return SnapshotAsset;
-}
-
-	return nullptr;
+	return SnapshotEditor::TakeLevelSnapshotAndSaveToDisk(TargetWorld, FileName, FolderPath, Description, bShouldCreateUniqueFileName, false);
 }
 
 void ULevelSnapshotsEditorFunctionLibrary::TakeAndSaveLevelSnapshotEditorWorld(const FString& FileName, const FString& FolderPath, const FString& Description)
@@ -140,36 +44,5 @@ void ULevelSnapshotsEditorFunctionLibrary::TakeAndSaveLevelSnapshotEditorWorld(c
 
 void ULevelSnapshotsEditorFunctionLibrary::GenerateThumbnailForSnapshotAsset(ULevelSnapshot* Snapshot)
 {
-	if (!ensure(Snapshot))
-	{
-		return;
-	}
-
-	FAssetRegistryModule& AssetRegistryModule = FModuleManager::LoadModuleChecked<FAssetRegistryModule>(TEXT("AssetRegistry"));
-	IAssetRegistry& AssetRegistry = AssetRegistryModule.Get();
-	TArray<FAssetData> SnapshotAssetData;
-	AssetRegistry.GetAssetsByPackageName(FName(*Snapshot->GetPackage()->GetPathName()), SnapshotAssetData);
-	
-	if (ensureMsgf(SnapshotAssetData.Num(), TEXT("Failed to find asset data for asset we just saved. Investigate!")))
-	{
-		// Copied from FAssetFileContextMenu::ExecuteCaptureThumbnail
-		FViewport* Viewport = GEditor->GetActiveViewport();
-
-		if ( ensure(GCurrentLevelEditingViewportClient) && ensure(Viewport) )
-		{
-			//have to re-render the requested viewport
-			FLevelEditorViewportClient* OldViewportClient = GCurrentLevelEditingViewportClient;
-			//remove selection box around client during render
-			GCurrentLevelEditingViewportClient = NULL;
-			Viewport->Draw();
-
-			AssetViewUtils::CaptureThumbnailFromViewport(Viewport, SnapshotAssetData);
-
-			//redraw viewport to have the yellow highlight again
-			GCurrentLevelEditingViewportClient = OldViewportClient;
-			Viewport->Draw();
-		}
-	}
+	SnapshotEditor::GenerateThumbnailForSnapshotAsset(Snapshot);
 }
-
-#undef LOCTEXT_NAMESPACE
diff --git a/Engine/Plugins/VirtualProduction/LevelSnapshots/Source/LevelSnapshotsEditor/Private/LevelSnapshotsEditorModule.cpp b/Engine/Plugins/VirtualProduction/LevelSnapshots/Source/LevelSnapshotsEditor/Private/LevelSnapshotsEditorModule.cpp
--- a/Engine/Plugins/VirtualProduction/LevelSnapshots/Source/LevelSnapshotsEditor/Private/LevelSnapshotsEditorModule.cpp
+++ b/Engine/Plugins/VirtualProduction/LevelSnapshots/Source/LevelSnapshotsEditor/Private/LevelSnapshotsEditorModule.cpp
@@ -2,36 +2,62 @@
 
 #include "LevelSnapshotsEditorModule.h"
 
-#include "Settings/LevelSnapshotsEditorProjectSettings.h"
-#include "Settings/LevelSnapshotsEditorDataManagementSettings.h"
-#include "NegatableFilter.h"
+#include "Data/Filters/NegatableFilter.h"
+#include "Data/LevelSnapshotsEditorData.h"
 #include "LevelSnapshotsEditorCommands.h"
 #include "LevelSnapshotsEditorStyle.h"
-#include "LevelSnapshotsEditorFunctionLibrary.h"
-#include "LevelSnapshotsEditorData.h"
-#include "LevelSnapshotsFunctionLibrary.h"
 #include "LevelSnapshotsUserSettings.h"
-#include "SLevelSnapshotsEditorCreationForm.h"
+#include "Settings/LevelSnapshotsEditorProjectSettings.h"
+#include "Settings/LevelSnapshotsEditorDataManagementSettings.h"
 #include "Toolkits/LevelSnapshotsEditorToolkit.h"
 
 #include "AssetRegistry/AssetRegistryModule.h"
 #include "AssetToolsModule.h"
 #include "AssetTypeActions/AssetTypeActions_LevelSnapshot.h"
+#include "GameProjectGenerationModule.h"
 #include "IAssetTools.h"
+#include "ILevelSnapshotsModule.h"
+#include "Interfaces/IProjectManager.h"
 #include "ISettingsModule.h"
 #include "LevelEditor.h"
+#include "Util/TakeSnapshotUtil.h"
+#include "UnrealEdMisc.h"
 #include "ToolMenus.h"
-#include "Framework/Notifications/NotificationManager.h"
 #include "Misc/ScopeExit.h"
-#include "Widgets/Notifications/SNotificationList.h"
+#include "Misc/MessageDialog.h"
 
 #define LOCTEXT_NAMESPACE "FLevelSnapshotsEditorModule"
 
+namespace
+{
+	bool EnableSupportPlugin(const FString& PluginName)
+	{
+		FText FailMessage = FText::GetEmpty();
+		bool bSuccess = IProjectManager::Get().SetPluginEnabled(PluginName, true, FailMessage);
+		if (bSuccess && IProjectManager::Get().IsCurrentProjectDirty())
+		{
+			FGameProjectGenerationModule::Get().TryMakeProjectFileWriteable(FPaths::GetProjectFilePath());
+			bSuccess = IProjectManager::Get().SaveCurrentProjectToDisk(FailMessage);
+		}
+
+		if (!bSuccess)
+		{
+			FMessageDialog::Open(EAppMsgType::Ok, FailMessage);
+		}
+		return bSuccess;
+	}
+}
+
 FLevelSnapshotsEditorModule& FLevelSnapshotsEditorModule::Get()
 {
 	return FModuleManager::GetModuleChecked<FLevelSnapshotsEditorModule>("LevelSnapshotsEditor");
 }
 
+void FLevelSnapshotsEditorModule::OpenLevelSnapshotsSettings()
+{
+	FModuleManager::LoadModuleChecked<ISettingsModule>("Settings").ShowViewer("Project", "Plugins", "Level Snapshots");
+}
+
 void FLevelSnapshotsEditorModule::StartupModule()
 {
 	IAssetTools& AssetTools = FModuleManager::LoadModuleChecked<FAssetToolsModule>("AssetTools").Get();
@@ -92,6 +118,7 @@ void FLevelSnapshotsEditorModule::SetUseCreationForm(bool bInUseCreationForm)
 void FLevelSnapshotsEditorModule::PostEngineInit()
 {
 	RegisterMenuItem();
+	SetupReminderToEnableSupportPlugins();
 
 	if (RegisterProjectSettings() && ProjectSettingsObjectPtr->bEnableLevelSnapshotsToolbarButton)
 	{
@@ -99,6 +126,60 @@ void FLevelSnapshotsEditorModule::PostEngineInit()
 	}
 }
 
+void FLevelSnapshotsEditorModule::SetupReminderToEnableSupportPlugins()
+{
+	ILevelSnapshotsModule& Module = ILevelSnapshotsModule::Get();
+	Module.AddCanTakeSnapshotDelegate(
+		"HasSupportPlugins",
+		ILevelSnapshotsModule::FCanTakeSnapshot::CreateLambda([](const FPreTakeSnapshotEventData&)
+		{
+			FModuleManager& ModuleManager = FModuleManager::Get();
+			bool bDidAnyPluginFail = false;
+			bool bNeedsToRestartEditor = false;
+			
+			const bool bShouldEnablenDisplaySupport =
+				ModuleManager.IsModuleLoaded("DisplayCluster")
+				&& !ModuleManager.IsModuleLoaded("nDisplaySupportForLevelSnapshots");
+			if (bShouldEnablenDisplaySupport)
+			{
+				const EAppReturnType::Type DialogResult = FMessageDialog::Open(
+					EAppMsgType::YesNoCancel,
+					LOCTEXT("EnableSupportFornDisplay", "The nDisplay plugin is enabled. To support snapshots of nDisplay actors that the scene may contain, you must enable the \"nDisplay Support For Level Snapshots\" plugin.\n\nDo you want to enable the plugin?")
+					);
+				if (DialogResult == EAppReturnType::Cancel)
+				{
+					// Cancel saving snapshot: this will stop saving the snapshot
+					return false;
+				}
+				if (DialogResult == EAppReturnType::Yes)
+				{
+					const bool bSuccess = EnableSupportPlugin("nDisplaySupportForLevelSnapshots");
+					bDidAnyPluginFail |= !bSuccess;
+					bNeedsToRestartEditor |= bSuccess;
+				}
+			}
+
+			if (bNeedsToRestartEditor)
+			{
+				FText RestartEditorTitle = LOCTEXT("RestartEditorTitle", "Restart required");
+				const EAppReturnType::Type RestartEditorResult = FMessageDialog::Open(
+					EAppMsgType::YesNo,
+					LOCTEXT("RestartEditor", "The editor needs to be restarted to enable the required plugin(s).\n\nRestart the editor now?"),
+					&RestartEditorTitle
+					);
+				if (RestartEditorResult == EAppReturnType::Yes)
+				{
+					const bool bWarn = false;
+					FUnrealEdMisc::Get().RestartEditor(bWarn);
+				}
+			}
+
+			// Editor restart does not start immediately: cancel saving the snapshot if restart is required
+			return !bDidAnyPluginFail && !bNeedsToRestartEditor;
+		})
+		);
+}
+
 void FLevelSnapshotsEditorModule::RegisterMenuItem()
 {
 	if (FSlateApplication::IsInitialized())
@@ -229,7 +310,7 @@ void FLevelSnapshotsEditorModule::MapEditorToolbarActions()
 void FLevelSnapshotsEditorModule::CreateEditorToolbarButton(FToolBarBuilder& Builder)
 {
 	Builder.AddToolBarButton(
-		FUIAction(FExecuteAction::CreateRaw(this, &FLevelSnapshotsEditorModule::BuildPathsToSaveSnapshotWithOptionalForm)),
+		FUIAction(FExecuteAction::CreateStatic(&SnapshotEditor::TakeSnapshotWithOptionalForm)),
 		NAME_None,
 		NSLOCTEXT("LevelSnapshots", "LevelSnapshots", "Level Snapshots"), // Set Text under image
 		NSLOCTEXT("LevelSnapshots", "LevelSnapshotsToolbarButtonTooltip", "Take snapshot with optional form"), //  Set tooltip
@@ -239,8 +320,8 @@ void FLevelSnapshotsEditorModule::CreateEditorToolbarButton(FToolBarBuilder& Bui
 	Builder.AddComboButton(
 		FUIAction(),
 		FOnGetContent::CreateRaw(this, &FLevelSnapshotsEditorModule::FillEditorToolbarComboButtonMenuOptions, EditorToolbarButtonCommandList), // Add combo button subcommands menu
-		FText::GetEmpty(),
-		NSLOCTEXT("LevelSnapshots", "LevelSnapshotsToolbarComboButtonTooltip", "Open Level Snapshots Options") //  Set tooltip
+		NSLOCTEXT("LevelSnapshots", "LevelSnapshotsOptions_Label", "Level Snapshots Options"), // Set text seen when the Level Editor Toolbar is truncated and the flyout is clicked
+		NSLOCTEXT("LevelSnapshots", "LevelSnapshotsToolbarComboButtonTooltip", "Open Level Snapshots Options"), FSlateIcon(), true //  Set tooltip
 	); 
 }
 
@@ -261,63 +342,6 @@ TSharedRef<SWidget> FLevelSnapshotsEditorModule::FillEditorToolbarComboButtonMen
 	return MenuBuilder.MakeWidget();
 }
 
-void FLevelSnapshotsEditorModule::BuildPathsToSaveSnapshotWithOptionalForm() const
-{	
-	check(ProjectSettingsObjectPtr.IsValid());
-	check(DataMangementSettingsObjectPtr.IsValid());
-
-	// Creation Form
-
-	if (ProjectSettingsObjectPtr.Get()->bUseCreationForm)
-	{
-		TSharedRef<SWidget> CreationForm = 
-			FLevelSnapshotsEditorCreationForm::MakeAndShowCreationWindow(
-				FCloseCreationFormDelegate::CreateRaw(this, &FLevelSnapshotsEditorModule::HandleFormReply), 
-				ProjectSettingsObjectPtr.Get(), DataMangementSettingsObjectPtr.Get());
-	}
-	else
-	{
-		TakeAndSaveSnapshot(FText::GetEmpty());
-	}
-}
-
-void FLevelSnapshotsEditorModule::HandleFormReply(bool bWasCreateSnapshotPressed, FText InDescription) const
-{
-	if (bWasCreateSnapshotPressed)
-	{
-		TakeAndSaveSnapshot(InDescription, true);
-	}
-}
-
-void FLevelSnapshotsEditorModule::TakeAndSaveSnapshot(const FText& InDescription, const bool bShouldUseOverrides) const
-{
-	if (!ensure(GEditor))
-	{
-		return;
-	}
-	UWorld* World = ULevelSnapshotsEditorData::GetEditorWorld();
-	if (!ensure(World && ProjectSettingsObjectPtr.IsValid()))
-	{
-		return;
-	}
-	ULevelSnapshotsEditorDataManagementSettings* DataManagementSettings = DataMangementSettingsObjectPtr.Get();
-
-
-	DataManagementSettings->ValidateRootLevelSnapshotSaveDirAsGameContentRelative();
-	DataManagementSettings->SanitizeAllProjectSettingsPaths(true);
-
-	const FText& NewSnapshotDir = ULevelSnapshotsEditorDataManagementSettings::ParseLevelSnapshotsTokensInText(
-		FText::FromString(
-			FPaths::Combine(DataManagementSettings->RootLevelSnapshotSaveDir.Path, DataManagementSettings->LevelSnapshotSaveDir)), World->GetName());
-	const FText& NewSnapshotName = ULevelSnapshotsEditorDataManagementSettings::ParseLevelSnapshotsTokensInText(FText::FromString(
-		bShouldUseOverrides && DataManagementSettings->IsNameOverridden() ?
-		DataManagementSettings->GetNameOverride() : DataManagementSettings->DefaultLevelSnapshotName), World->GetName());
-
-	const FString& ValidatedName = FPaths::MakeValidFileName(NewSnapshotName.ToString());
-
-	ULevelSnapshotsEditorFunctionLibrary::TakeLevelSnapshotAndSaveToDisk(World, ValidatedName, NewSnapshotDir.ToString(), InDescription.ToString(), false);
-}
-
 void FLevelSnapshotsEditorModule::OpenLevelSnapshotsDialogWithAssetSelected(const FAssetData& InAssetData)
 {
 	OpenSnapshotsEditor();
@@ -341,11 +365,6 @@ void FLevelSnapshotsEditorModule::OpenSnapshotsEditor()
 	}
 }
 
-void FLevelSnapshotsEditorModule::OpenLevelSnapshotsSettings()
-{
-	FModuleManager::LoadModuleChecked<ISettingsModule>("Settings").ShowViewer("Project", "Plugins", "Level Snapshots");
-}
-
 ULevelSnapshotsEditorData* FLevelSnapshotsEditorModule::AllocateTransientPreset()
 {
 	static const TCHAR* PackageName = TEXT("/Temp/LevelSnapshots/PendingSnapshots");
diff --git a/Engine/Plugins/VirtualProduction/LevelSnapshots/Source/LevelSnapshotsEditor/Private/LevelSnapshotsEditorModule.h b/Engine/Plugins/VirtualProduction/LevelSnapshots/Source/LevelSnapshotsEditor/Private/LevelSnapshotsEditorModule.h
--- a/Engine/Plugins/VirtualProduction/LevelSnapshots/Source/LevelSnapshotsEditor/Private/LevelSnapshotsEditorModule.h
+++ b/Engine/Plugins/VirtualProduction/LevelSnapshots/Source/LevelSnapshotsEditor/Private/LevelSnapshotsEditorModule.h
@@ -19,6 +19,7 @@ class FLevelSnapshotsEditorModule : public IModuleInterface
 public:
 
 	static FLevelSnapshotsEditorModule& Get();
+	static void OpenLevelSnapshotsSettings();
 
 	//~ Begin IModuleInterface Interface
 	virtual void StartupModule() override;
@@ -26,38 +27,20 @@ public:
 	//~ End IModuleInterface Interface
 
 	bool GetUseCreationForm() const;
-
 	void SetUseCreationForm(bool bInUseCreationForm);
-
-	void ToggleUseCreationForm()
-	{
-		SetUseCreationForm(!GetUseCreationForm());
-	}
-
-	void BuildPathsToSaveSnapshotWithOptionalForm() const;
-
-	void HandleFormReply(bool bWasCreateSnapshotPressed, FText InDescription) const;
-
-	void TakeAndSaveSnapshot(const FText& InDescription, const bool bShouldUseOverrides = false) const;
+	void ToggleUseCreationForm() { SetUseCreationForm(!GetUseCreationForm()); }
 
 	void OpenLevelSnapshotsDialogWithAssetSelected(const FAssetData& InAssetData);
-	
-	static void OpenLevelSnapshotsSettings();
 
-	TWeakObjectPtr<ULevelSnapshotsEditorProjectSettings> GetLevelSnapshotsUserSettings() const
-	{
-		return ProjectSettingsObjectPtr;
-	}
-	
-	TWeakObjectPtr<ULevelSnapshotsEditorDataManagementSettings> GetLevelSnapshotsDataManagementSettings() const
-	{
-		return DataMangementSettingsObjectPtr;
-	}
+	TWeakObjectPtr<ULevelSnapshotsEditorProjectSettings> GetLevelSnapshotsUserSettings() const { return ProjectSettingsObjectPtr; }
+	TWeakObjectPtr<ULevelSnapshotsEditorDataManagementSettings> GetLevelSnapshotsDataManagementSettings() const { return DataMangementSettingsObjectPtr; }
 
 private:
 	
 	void PostEngineInit();
 
+	void SetupReminderToEnableSupportPlugins();
+
 	void RegisterMenuItem();
 	bool RegisterProjectSettings();
 	bool HandleModifiedProjectSettings();
@@ -68,13 +51,13 @@ private:
 	TSharedRef<SWidget> FillEditorToolbarComboButtonMenuOptions(TSharedPtr<class FUICommandList> Commands);
 
 	void OpenSnapshotsEditor();
-
 	ULevelSnapshotsEditorData* AllocateTransientPreset();
 
-	// Command list (for combo button sub menu options)
+	
+	/** Command list (for combo button sub menu options) */
 	TSharedPtr<FUICommandList> EditorToolbarButtonCommandList;
 
-	/* Lives for as long as the UI is open. */
+	/** Lives for as long as the UI is open. */
 	TWeakPtr<FLevelSnapshotsEditorToolkit> SnapshotEditorToolkit;
 
 	TSharedPtr<ISettingsSection> ProjectSettingsSectionPtr;
diff --git a/Engine/Plugins/VirtualProduction/LevelSnapshots/Source/LevelSnapshotsEditor/Private/Toolkits/LevelSnapshotsEditorToolkit.cpp b/Engine/Plugins/VirtualProduction/LevelSnapshots/Source/LevelSnapshotsEditor/Private/Toolkits/LevelSnapshotsEditorToolkit.cpp
--- a/Engine/Plugins/VirtualProduction/LevelSnapshots/Source/LevelSnapshotsEditor/Private/Toolkits/LevelSnapshotsEditorToolkit.cpp
+++ b/Engine/Plugins/VirtualProduction/LevelSnapshots/Source/LevelSnapshotsEditor/Private/Toolkits/LevelSnapshotsEditorToolkit.cpp
@@ -17,6 +17,7 @@
 #include "EditorStyleSet.h"
 #include "Framework/Docking/LayoutService.h"
 #include "Stats/StatsMisc.h"
+#include "Util/TakeSnapshotUtil.h"
 #include "Widgets/Docking/SDockTab.h"
 #include "Widgets/Input/SButton.h"
 #include "Widgets/Input/SCheckBox.h"
@@ -301,7 +302,7 @@ TSharedRef<SDockTab> FLevelSnapshotsEditorToolkit::SpawnTab_Results(const FSpawn
 
 FReply FLevelSnapshotsEditorToolkit::OnClickTakeSnapshot()
 {
-	FLevelSnapshotsEditorModule::Get().BuildPathsToSaveSnapshotWithOptionalForm();
+	SnapshotEditor::TakeSnapshotWithOptionalForm();
 	return FReply::Handled();
 }
 
diff --git a/Engine/Plugins/VirtualProduction/LevelSnapshots/Source/LevelSnapshotsEditor/Private/Util/TakeSnapshotUtil.cpp b/Engine/Plugins/VirtualProduction/LevelSnapshots/Source/LevelSnapshotsEditor/Private/Util/TakeSnapshotUtil.cpp
new file mode 100644
--- /dev/null
+++ b/Engine/Plugins/VirtualProduction/LevelSnapshots/Source/LevelSnapshotsEditor/Private/Util/TakeSnapshotUtil.cpp
@@ -0,0 +1,236 @@
+﻿// Copyright Epic Games, Inc. All Rights Reserved.
+
+#include "Util/TakeSnapshotUtil.h"
+
+#include "LevelSnapshotsLog.h"
+
+#include "AssetRegistry/AssetRegistryModule.h"
+#include "AssetToolsModule.h"
+#include "AssetViewUtils.h"
+#include "IAssetTools.h"
+#include "Editor.h"
+#include "Engine/World.h"
+#include "LevelSnapshot.h"
+#include "LevelSnapshotsEditorModule.h"
+#include "LevelSnapshotsEditorStyle.h"
+#include "ObjectTools.h"
+#include "Data/LevelSnapshotsEditorData.h"
+#include "Framework/Notifications/NotificationManager.h"
+#include "GameplayMediaEncoder/Private/GameplayMediaEncoderCommon.h"
+#include "Logging/MessageLog.h"
+#include "Misc/FeedbackContext.h"
+#include "Misc/MessageDialog.h"
+#include "UObject/Package.h"
+#include "Widgets/SLevelSnapshotsEditorCreationForm.h"
+#include "Widgets/Notifications/SNotificationList.h"
+
+#define LOCTEXT_NAMESPACE "LevelSnapshotEditorLibrary"
+
+namespace
+{
+	void HandleFormReply(const FText& InDescription, bool bShouldUseOverrides, bool bSaveAsync)
+	{
+		FLevelSnapshotsEditorModule& Module = FLevelSnapshotsEditorModule::Get();
+		TWeakObjectPtr<ULevelSnapshotsEditorProjectSettings> ProjectSettings = Module.GetLevelSnapshotsUserSettings();
+		TWeakObjectPtr<ULevelSnapshotsEditorDataManagementSettings> DataMangementSettings = Module.GetLevelSnapshotsDataManagementSettings();
+		
+		UWorld* World = ULevelSnapshotsEditorData::GetEditorWorld();
+		if (!ensure(World && ProjectSettings.IsValid() && DataMangementSettings.IsValid()))
+		{
+			return;
+		}
+		
+		ULevelSnapshotsEditorDataManagementSettings* DataManagementSettings = DataMangementSettings.Get();
+		DataManagementSettings->ValidateRootLevelSnapshotSaveDirAsGameContentRelative();
+		DataManagementSettings->SanitizeAllProjectSettingsPaths(true);
+
+		const FText& NewSnapshotDir = ULevelSnapshotsEditorDataManagementSettings::ParseLevelSnapshotsTokensInText(
+			FText::FromString(FPaths::Combine(DataManagementSettings->RootLevelSnapshotSaveDir.Path, DataManagementSettings->LevelSnapshotSaveDir)),
+			World->GetName()
+			);
+		const FString& DescriptionString = bShouldUseOverrides && DataManagementSettings->IsNameOverridden() ?	DataManagementSettings->GetNameOverride() : DataManagementSettings->DefaultLevelSnapshotName;
+		const FText& NewSnapshotName = ULevelSnapshotsEditorDataManagementSettings::ParseLevelSnapshotsTokensInText(
+				FText::FromString(DescriptionString),
+				World->GetName()
+				);
+
+		const FString& ValidatedName = FPaths::MakeValidFileName(NewSnapshotName.ToString());
+
+		SnapshotEditor::TakeLevelSnapshotAndSaveToDisk(World, ValidatedName, NewSnapshotDir.ToString(), InDescription.ToString(), false, bSaveAsync);
+	}
+	
+	void DestroySnapshot(ULevelSnapshot* SnapshotAsset)
+	{
+		SnapshotAsset->ClearFlags(RF_Public | RF_Standalone);
+		SnapshotAsset->Rename(nullptr, GetTransientPackage());
+		SnapshotAsset->MarkPendingKill();
+		SnapshotAsset->RemoveFromRoot();
+	}
+}
+
+void SnapshotEditor::TakeSnapshotWithOptionalForm()
+{
+	FLevelSnapshotsEditorModule& Module = FLevelSnapshotsEditorModule::Get();
+	TWeakObjectPtr<ULevelSnapshotsEditorProjectSettings> ProjectSettings = Module.GetLevelSnapshotsUserSettings();
+	TWeakObjectPtr<ULevelSnapshotsEditorDataManagementSettings> DataMangementSettings = Module.GetLevelSnapshotsDataManagementSettings();
+	
+	if (ProjectSettings.Get()->bUseCreationForm)
+	{
+		TSharedRef<SWidget> CreationForm = SLevelSnapshotsEditorCreationForm::MakeAndShowCreationWindow(
+			FCloseCreationFormDelegate::CreateLambda([](const FText& Description, bool bSaveAsync)
+			{
+				HandleFormReply(Description, true, bSaveAsync);
+			}),
+			ProjectSettings.Get(),
+			DataMangementSettings.Get()
+			);
+	}
+	else
+	{
+		HandleFormReply(FText::GetEmpty(), false, false);
+	}
+}
+
+ULevelSnapshot* SnapshotEditor::TakeLevelSnapshotAndSaveToDisk(UWorld* World, const FString& FileName, const FString& FolderPath, const FString& Description, bool bShouldCreateUniqueFileName, bool bSaveAsync)
+{
+	SCOPED_SNAPSHOT_EDITOR_TRACE(CreateSnapshotPackage);
+	FString AssetName = FileName;
+    	
+    FString BasePackageName = FPaths::Combine(FolderPath, FileName);
+
+    BasePackageName.RemoveFromStart(TEXT("/"));
+    BasePackageName.RemoveFromStart(TEXT("Content/"));
+    BasePackageName.StartsWith(TEXT("Game/")) == true ? BasePackageName.InsertAt(0, TEXT("/")) : BasePackageName.InsertAt(0, TEXT("/Game/"));
+
+    FString PackageName = BasePackageName;
+
+    if (bShouldCreateUniqueFileName)
+    {
+    	IAssetTools& AssetTools = FModuleManager::Get().LoadModuleChecked<FAssetToolsModule>("AssetTools").Get();
+    	
+    	AssetTools.CreateUniqueAssetName(BasePackageName, TEXT(""), PackageName, AssetName);
+    }
+
+    const FString& PackageFileName = FPackageName::LongPackageNameToFilename(PackageName, FPackageName::GetAssetPackageExtension());
+    const bool bPathIsValid = FPaths::ValidatePath(PackageFileName);
+
+    bool bProceedWithSave = bPathIsValid;
+
+    if (bProceedWithSave && !bShouldCreateUniqueFileName)
+    {
+    	FAssetRegistryModule& AssetRegistryModule = FModuleManager::LoadModuleChecked<FAssetRegistryModule>("AssetRegistry");
+    	FAssetData ExistingAsset = AssetRegistryModule.Get().GetAssetByObjectPath(FName(*(BasePackageName + "." + FileName)));
+    	if ( ExistingAsset.IsValid() && ExistingAsset.AssetClass == *ULevelSnapshot::StaticClass()->GetName() )
+    	{
+    		EAppReturnType::Type ShouldReplace = FMessageDialog::Open( EAppMsgType::YesNo, FText::Format(LOCTEXT("ReplaceAssetMessage", "{0} already exists. Do you want to replace it?"), FText::FromString(FileName)) );
+    		bProceedWithSave = (ShouldReplace == EAppReturnType::Yes);
+    	}
+    }
+
+    if (bProceedWithSave)
+    {
+    	// Show notification
+    	FNotificationInfo Notification(NSLOCTEXT("LevelSnapshots", "NotificationFormatText_CreatingSnapshot", "Creating Level Snapshot"));
+    	Notification.Image = FLevelSnapshotsEditorStyle::GetBrush(TEXT("LevelSnapshots.ToolbarButton"));
+    	Notification.bUseThrobber = true;
+    	Notification.bUseSuccessFailIcons = true;
+    	Notification.ExpireDuration = 2.f;
+    	Notification.bFireAndForget = false;
+    	
+    	TSharedPtr<SNotificationItem> NotificationItem = FSlateNotificationManager::Get().AddNotification(Notification);
+    	NotificationItem->SetCompletionState(SNotificationItem::CS_Pending);
+    	ON_SCOPE_EXIT
+    	{
+    		NotificationItem->ExpireAndFadeout();
+    	};
+
+    	UPackage* SavePackage = CreatePackage(*PackageName);
+    	const EObjectFlags AssetFlags = RF_Public | RF_Standalone;
+    	ULevelSnapshot* SnapshotAsset = NewObject<ULevelSnapshot>(SavePackage, *AssetName, AssetFlags, nullptr);
+
+    	check(SavePackage && SnapshotAsset);
+    	
+    	SnapshotAsset->SetSnapshotName(*FileName);
+    	SnapshotAsset->SetSnapshotDescription(Description);
+
+        const bool bSuccessful = SnapshotAsset->SnapshotWorld(World);
+        if (!bSuccessful)
+        {
+            // Package and snapshot need to be destroyed again
+            DestroySnapshot(SnapshotAsset);
+            ObjectTools::DeleteObjectsUnchecked({ SavePackage });
+            
+            NotificationItem->SetText(
+                NSLOCTEXT("LevelSnapshots", "NotificationFormatText_NoSnapshotTaken", "No snapshot was taken")
+                );
+            NotificationItem->SetCompletionState(SNotificationItem::CS_Fail);
+            return nullptr;
+        }
+    	SnapshotAsset->MarkPackageDirty();
+
+    	FAssetRegistryModule::AssetCreated(SnapshotAsset);
+    	GenerateThumbnailForSnapshotAsset(SnapshotAsset);
+
+    	bool bSavingSuccessful;
+    	{
+    		SCOPED_SNAPSHOT_EDITOR_TRACE(SaveSnapshotPackage);
+    		bSavingSuccessful = UPackage::SavePackage(SavePackage, SnapshotAsset, RF_Public | RF_Standalone, *PackageFileName, GWarn, nullptr, false, false,  bSaveAsync ? SAVE_Async | SAVE_NoError : SAVE_None);
+    	}
+    	
+    	// Notify the user of the outcome
+    	if (bPathIsValid && bSavingSuccessful)
+    	{
+    		NotificationItem->SetText(
+    			FText::Format(
+    				NSLOCTEXT("LevelSnapshots", "NotificationFormatText_CreateSnapshotSuccess", "Successfully created Level Snapshot \"{0}\""), FText::FromString(FileName)));
+    		NotificationItem->SetCompletionState(SNotificationItem::CS_Success);
+    	}
+    	else
+    	{
+    		NotificationItem->SetText(
+    			FText::Format(
+    				NSLOCTEXT("LevelSnapshots", "NotificationFormatText_CreateSnapshotSuccess", "Failed to create Level Snapshot \"{0}\". Check the file name."), FText::FromString(FileName)));
+    		NotificationItem->SetCompletionState(SNotificationItem::CS_Fail);
+    	}
+
+    	return SnapshotAsset;
+    }
+
+    return nullptr;
+}
+
+void SnapshotEditor::GenerateThumbnailForSnapshotAsset(ULevelSnapshot* Snapshot)
+{
+	if (!ensure(Snapshot))
+	{
+		return;
+	}
+	
+	FAssetRegistryModule& AssetRegistryModule = FModuleManager::LoadModuleChecked<FAssetRegistryModule>(TEXT("AssetRegistry"));
+	IAssetRegistry& AssetRegistry = AssetRegistryModule.Get();
+	TArray<FAssetData> SnapshotAssetData;
+	AssetRegistry.GetAssetsByPackageName(FName(*Snapshot->GetPackage()->GetPathName()), SnapshotAssetData);
+	
+	if (ensureMsgf(SnapshotAssetData.Num(), TEXT("Failed to find asset data for asset we just saved. Investigate!")))
+	{
+		// Copied from FAssetFileContextMenu::ExecuteCaptureThumbnail
+		FViewport* Viewport = GEditor->GetActiveViewport();
+
+		if (ensure(GCurrentLevelEditingViewportClient) && ensure(Viewport))
+		{
+			//have to re-render the requested viewport
+			FLevelEditorViewportClient* OldViewportClient = GCurrentLevelEditingViewportClient;
+			//remove selection box around client during render
+			GCurrentLevelEditingViewportClient = NULL;
+			Viewport->Draw();
+
+			AssetViewUtils::CaptureThumbnailFromViewport(Viewport, SnapshotAssetData);
+
+			//redraw viewport to have the yellow highlight again
+			GCurrentLevelEditingViewportClient = OldViewportClient;
+			Viewport->Draw();
+		}
+	}
+}
+
+#undef LOCTEXT_NAMESPACE
\ No newline at end of file
diff --git a/Engine/Plugins/VirtualProduction/LevelSnapshots/Source/LevelSnapshotsEditor/Private/Util/TakeSnapshotUtil.h b/Engine/Plugins/VirtualProduction/LevelSnapshots/Source/LevelSnapshotsEditor/Private/Util/TakeSnapshotUtil.h
new file mode 100644
--- /dev/null
+++ b/Engine/Plugins/VirtualProduction/LevelSnapshots/Source/LevelSnapshotsEditor/Private/Util/TakeSnapshotUtil.h
@@ -0,0 +1,30 @@
+﻿// Copyright Epic Games, Inc. All Rights Reserved.
+
+#pragma once
+
+#include "CoreMinimal.h"
+
+class ULevelSnapshot;
+
+namespace SnapshotEditor
+{
+	/**
+	 * Shows the window for taking snapshots, if enabled by the user settings. Otherwise just takes a snapshot.
+	 */
+	void TakeSnapshotWithOptionalForm();
+
+	/**
+	 * @param World The world to take snapshot in
+	 * @param FileName Name of saved snapshot asset
+	 * @param FolderPath Where to save the asset
+	 * @param Description The value of the description field of the snapshot
+	 * @param bShouldCreateUniqueFileName Whether to use FileName as base to generate a unique file name
+	 * @param bSaveAsync Whether to save the asset async. Saving sync will block user and this can take several minutes.
+	 */
+	ULevelSnapshot* TakeLevelSnapshotAndSaveToDisk(UWorld* World, const FString& FileName, const FString& FolderPath, const FString& Description, bool bShouldCreateUniqueFileName = true, bool bSaveAsync = false);
+
+	/**
+	 * Captures a screenshot from the editor camera and assigns it to the snapshot.
+	 */
+	void GenerateThumbnailForSnapshotAsset(ULevelSnapshot* Snapshot);
+}
diff --git a/Engine/Plugins/VirtualProduction/LevelSnapshots/Source/LevelSnapshotsEditor/Private/Views/Input/SLevelSnapshotsEditorInput.cpp b/Engine/Plugins/VirtualProduction/LevelSnapshots/Source/LevelSnapshotsEditor/Private/Views/Input/SLevelSnapshotsEditorInput.cpp
--- a/Engine/Plugins/VirtualProduction/LevelSnapshots/Source/LevelSnapshotsEditor/Private/Views/Input/SLevelSnapshotsEditorInput.cpp
+++ b/Engine/Plugins/VirtualProduction/LevelSnapshots/Source/LevelSnapshotsEditor/Private/Views/Input/SLevelSnapshotsEditorInput.cpp
@@ -2,6 +2,7 @@
 
 #include "Views/Input/SLevelSnapshotsEditorInput.h"
 
+#include "Editor.h"
 #include "ILevelSnapshotsEditorView.h"
 #include "Views/Input/LevelSnapshotsEditorInput.h"
 #include "LevelSnapshotsEditorData.h"
@@ -17,7 +18,7 @@
 
 SLevelSnapshotsEditorInput::~SLevelSnapshotsEditorInput()
 {
-	FCoreUObjectDelegates::OnObjectModified.Remove(OnObjectModifiedDelegateHandle);
+	FEditorDelegates::OnMapOpened.Remove(OnMapOpenedDelegateHandle);
 }
 
 void SLevelSnapshotsEditorInput::Construct(const FArguments& InArgs, const TSharedRef<FLevelSnapshotsEditorInput>& InEditorInput, const TSharedRef<FLevelSnapshotsEditorViewBuilder>& InBuilder)
@@ -27,16 +28,24 @@ void SLevelSnapshotsEditorInput::Construct(const FArguments& InArgs, const TShar
 
 	check(BuilderPtr.Pin()->EditorDataPtr.IsValid());
 
+	OnMapOpenedDelegateHandle = FEditorDelegates::OnMapOpened.AddLambda([this] (const FString& InFileName, const bool bAsTemplate)
+	{
+		check(BuilderPtr.IsValid());
+		check(BuilderPtr.Pin()->EditorDataPtr.IsValid());
+	
+		OverrideWorld(BuilderPtr.Pin()->EditorDataPtr->GetEditorWorld());
+	});
+
 	ChildSlot
-		[
-			SAssignNew(EditorInputOuterVerticalBox, SVerticalBox)
+	[
+		SAssignNew(EditorInputOuterVerticalBox, SVerticalBox)
 
-			+ SVerticalBox::Slot()
-			[
-				SAssignNew(EditorBrowserWidgetPtr, SLevelSnapshotsEditorBrowser, InBuilder)
-				.OwningWorldPath(BuilderPtr.Pin()->EditorDataPtr->GetEditorWorld())
-			]
-		];
+		+ SVerticalBox::Slot()
+		[
+			SAssignNew(EditorBrowserWidgetPtr, SLevelSnapshotsEditorBrowser, InBuilder)
+			.OwningWorldPath(BuilderPtr.Pin()->EditorDataPtr->GetEditorWorld())
+		]
+	];
 }
 
 void SLevelSnapshotsEditorInput::OpenLevelSnapshotsDialogWithAssetSelected(const FAssetData& InAssetData) const
diff --git a/Engine/Plugins/VirtualProduction/LevelSnapshots/Source/LevelSnapshotsEditor/Private/Views/Input/SLevelSnapshotsEditorInput.h b/Engine/Plugins/VirtualProduction/LevelSnapshots/Source/LevelSnapshotsEditor/Private/Views/Input/SLevelSnapshotsEditorInput.h
--- a/Engine/Plugins/VirtualProduction/LevelSnapshots/Source/LevelSnapshotsEditor/Private/Views/Input/SLevelSnapshotsEditorInput.h
+++ b/Engine/Plugins/VirtualProduction/LevelSnapshots/Source/LevelSnapshotsEditor/Private/Views/Input/SLevelSnapshotsEditorInput.h
@@ -31,7 +31,7 @@ public:
 private:
 	void OverrideWorld(FSoftObjectPath InNewContextPath);
 
-	FDelegateHandle OnObjectModifiedDelegateHandle;
+	FDelegateHandle OnMapOpenedDelegateHandle;
 	
 	TWeakPtr<FLevelSnapshotsEditorInput> EditorInputPtr;
 
diff --git a/Engine/Plugins/VirtualProduction/LevelSnapshots/Source/LevelSnapshotsEditor/Private/Views/Results/SLevelSnapshotsEditorResults.cpp b/Engine/Plugins/VirtualProduction/LevelSnapshots/Source/LevelSnapshotsEditor/Private/Views/Results/SLevelSnapshotsEditorResults.cpp
--- a/Engine/Plugins/VirtualProduction/LevelSnapshots/Source/LevelSnapshotsEditor/Private/Views/Results/SLevelSnapshotsEditorResults.cpp
+++ b/Engine/Plugins/VirtualProduction/LevelSnapshots/Source/LevelSnapshotsEditor/Private/Views/Results/SLevelSnapshotsEditorResults.cpp
@@ -8,7 +8,6 @@
 #include "Data/LevelSnapshot.h"
 #include "LevelSnapshotsLog.h"
 #include "LevelSnapshotsEditorStyle.h"
-#include "LevelSnapshotsStats.h"
 #include "PropertyInfoHelpers.h"
 #include "PropertySelection.h"
 
@@ -279,7 +278,7 @@ void FLevelSnapshotsEditorResultsRow::GenerateActorGroupChildren(FPropertySelect
 					const TSharedRef<FPropertyHandleHierarchy>& ChildHierarchy = WorldHandleHierarchy->DirectChildren[ChildIndex];
 
 					LoopOverHandleHierarchiesAndCreateRowHierarchy(
-							ChildHierarchy, InDirectParentRow, PropertiesThatPassFilter, PropertyRowsGenerated, InResultsView, SnapshotHandleHierarchy);
+							ObjectType_World, ChildHierarchy, InDirectParentRow, PropertiesThatPassFilter, PropertyRowsGenerated, InResultsView, SnapshotHandleHierarchy);
 				}
 			}
 
@@ -287,41 +286,42 @@ void FLevelSnapshotsEditorResultsRow::GenerateActorGroupChildren(FPropertySelect
 		}
 
 		/* Do not pass in the base hierarchy, only children of the base hierarchy. */
-		static void LoopOverHandleHierarchiesAndCreateRowHierarchy(const TWeakPtr<FPropertyHandleHierarchy>& InHierarchy,
+		static void LoopOverHandleHierarchiesAndCreateRowHierarchy(ELevelSnapshotsObjectType InputType, const TWeakPtr<FPropertyHandleHierarchy>& InputHierarchy,
 			const TWeakPtr<FLevelSnapshotsEditorResultsRow>& InDirectParentRow, const TArray<TFieldPath<FProperty>>& PropertiesThatPassFilter,
-			TArray<TFieldPath<FProperty>>& PropertyRowsGenerated, const TWeakPtr<SLevelSnapshotsEditorResults>& InResultsView, const TWeakPtr<FPropertyHandleHierarchy>& InHierarchyToSearchForCounterparts = nullptr)
+			TArray<TFieldPath<FProperty>>& PropertyRowsGenerated, const TWeakPtr<SLevelSnapshotsEditorResults>& InResultsView, 
+			const TWeakPtr<FPropertyHandleHierarchy>& InHierarchyToSearchForCounterparts = nullptr)
 		{
-			if (!ensure(InHierarchy.IsValid()) || !ensure(InDirectParentRow.IsValid()))
+			if (!ensureAlwaysMsgf(InputType != ObjectType_None, 
+				TEXT("%hs: InputType must NOT be ObjectType_None!"), __FUNCTION__))
 			{
 				return;
 			}
 			
-			const TSharedPtr<FPropertyHandleHierarchy> PinnedHierarchy = InHierarchy.Pin();
+			if (!ensureAlwaysMsgf(InputHierarchy.IsValid(), TEXT("%hs: InputHierarchy must be valid."), __FUNCTION__) || 
+				!ensureAlwaysMsgf(InDirectParentRow.IsValid(), TEXT("%hs: InDirectParentRow must be valid."), __FUNCTION__))
+			{
+				return;
+			}
+
+			const bool bInputIsWorld = InputType == ELevelSnapshotsObjectType::ObjectType_World;
+			
+			const TSharedPtr<FPropertyHandleHierarchy> PinnedInputHierarchy = InputHierarchy.Pin();
 
-			const TSharedPtr<IPropertyHandle>& Handle = PinnedHierarchy->Handle;
+			const TSharedPtr<IPropertyHandle>& InputHandle = PinnedInputHierarchy->Handle;
 
 			// No asserts for handle or property because PropertyRowGenerator assumes there will be a details layout pointer but this scenario doesn't create one
 			// Asserting would cause needless debugging breaks every time a snapshot is chosen
-			if (!Handle->IsValidHandle())
+			if (!InputHandle->IsValidHandle())
 			{
 				return;
 			}
 
-			FProperty* Property = Handle->GetProperty();
+			FProperty* Property = InputHandle->GetProperty();
 			if (!Property)
 			{
 				return;
 			}
 
-			uint32 ChildHandleCount;
-			Handle->GetNumChildren(ChildHandleCount);
-			const FLevelSnapshotsEditorResultsRow::ELevelSnapshotsEditorResultsRowType InRowType = 
-				FLevelSnapshotsEditorResultsRow::DetermineRowTypeFromProperty(Property, Handle->IsCustomized(), ChildHandleCount > 0);
-			if (!ensure(InRowType != FLevelSnapshotsEditorResultsRow::None))
-			{
-				return;
-			}
-
 			const TFieldPath<FProperty> PropertyField(Property);
 
 			const FLevelSnapshotsEditorResultsRow::ELevelSnapshotsEditorResultsRowType ParentRowType = InDirectParentRow.Pin()->GetRowType();
@@ -330,15 +330,70 @@ void FLevelSnapshotsEditorResultsRow::GenerateActorGroupChildren(FPropertySelect
 					ParentRowType == FLevelSnapshotsEditorResultsRow::StructGroup || 
 					ParentRowType == FLevelSnapshotsEditorResultsRow::StructInSetOrArray ||
 					ParentRowType == FLevelSnapshotsEditorResultsRow::StructInMap);
+			
 			const bool bIsPropertyFilteredOut = !bIsParentRowContainer && !PropertiesThatPassFilter.Contains(PropertyField);
 			if (bIsPropertyFilteredOut) { return; }
+			
+			bool bFoundCounterpart = false;
+			TSharedPtr<FPropertyHandleHierarchy> PinnedCounterpartHierarchy;
+			TSharedPtr<IPropertyHandle> CounterpartHandle;
+			
+			if (InHierarchyToSearchForCounterparts.IsValid())
+			{
+				const TWeakPtr<FPropertyHandleHierarchy> OutCounterpartHierarchy = FindCorrespondingHandle(
+					PinnedInputHierarchy->TempIdentifierChain, PinnedInputHierarchy->DisplayName, InHierarchyToSearchForCounterparts, bFoundCounterpart);
 
-			bool bIsCounterpartValueSame = false;
+				if (bFoundCounterpart)
+				{
+					bFoundCounterpart = OutCounterpartHierarchy.IsValid();
+
+					if (bFoundCounterpart)
+					{
+						PinnedCounterpartHierarchy = OutCounterpartHierarchy.Pin();
+						
+						CounterpartHandle = PinnedCounterpartHierarchy->Handle;
+					}
+				}
+			}
+			
+			uint32 InputChildHandleCount = 0;
+			uint32 CounterpartChildHandleCount = 0;
+			InputHandle->GetNumChildren(InputChildHandleCount);
+
+			const int32 InputHierarchyChildCount = PinnedInputHierarchy->DirectChildren.Num();
+			int32 CounterpartHierarchyChildCount = 0;
+
+			if (bFoundCounterpart && PinnedCounterpartHierarchy.IsValid())
+			{
+				if (CounterpartHandle.IsValid())
+				{
+					CounterpartHandle->GetNumChildren(CounterpartChildHandleCount);
+				}
+				
+				CounterpartHierarchyChildCount = PinnedCounterpartHierarchy->DirectChildren.Num();
+			}
 
-			TWeakPtr<FPropertyHandleHierarchy> FoundHierarchy = nullptr;
+			const int32 MaxHandleChildCount = FMath::Max((int32)InputChildHandleCount, (int32)CounterpartChildHandleCount);
+			int32 MaxHierarchyChildCount = FMath::Max(InputHierarchyChildCount, CounterpartHierarchyChildCount);
+			
+			const int32 MaxChildCount = FMath::Max(MaxHandleChildCount, MaxHierarchyChildCount);
+			
+			const FLevelSnapshotsEditorResultsRow::ELevelSnapshotsEditorResultsRowType InRowType = 
+				FLevelSnapshotsEditorResultsRow::DetermineRowTypeFromProperty(Property, InputHandle->IsCustomized(), MaxChildCount > 0);
+			if (!ensure(InRowType != FLevelSnapshotsEditorResultsRow::None))
+			{
+				return;
+			}
+
+			bool bIsCounterpartValueSame = false;
 
-			auto AreHandleValuesEqual = [](const TSharedPtr<IPropertyHandle>& HandleA, const TSharedPtr<IPropertyHandle>& HandleB)
+			auto AreHandleValuesEqual = [](const TSharedPtr<IPropertyHandle> HandleA, const TSharedPtr<IPropertyHandle> HandleB)
 			{
+				if (!HandleA.IsValid() || !HandleB.IsValid())
+				{
+					return false;
+				}
+				
 				FString ValueA;
 				FString ValueB;
 
@@ -348,103 +403,120 @@ void FLevelSnapshotsEditorResultsRow::GenerateActorGroupChildren(FPropertySelect
 				return ValueA.Equals(ValueB);
 			};
 			
-			if (InHierarchyToSearchForCounterparts.IsValid())
+			if (bFoundCounterpart && CounterpartHandle.IsValid() &&
+				(InRowType == FLevelSnapshotsEditorResultsRow::SinglePropertyInMap || 
+				InRowType == FLevelSnapshotsEditorResultsRow::SinglePropertyInSetOrArray || 
+				InRowType == FLevelSnapshotsEditorResultsRow::SinglePropertyInStruct))
 			{
-				bool bFoundMatch = false;
-				FoundHierarchy = FindCorrespondingHandle(
-					PinnedHierarchy->TempIdentifierChain, PinnedHierarchy->DisplayName, InHierarchyToSearchForCounterparts, bFoundMatch);
-
-				if (bFoundMatch && FoundHierarchy.IsValid() && FoundHierarchy.Pin()->Handle.IsValid() &&
-					(InRowType == FLevelSnapshotsEditorResultsRow::SinglePropertyInMap || 
-					InRowType == FLevelSnapshotsEditorResultsRow::SinglePropertyInSetOrArray || 
-					InRowType == FLevelSnapshotsEditorResultsRow::SinglePropertyInStruct))
-				{
-					bIsCounterpartValueSame = AreHandleValuesEqual(Handle, FoundHierarchy.Pin()->Handle);
-				}
+				bIsCounterpartValueSame = AreHandleValuesEqual(InputHandle, CounterpartHandle);
 			}
 
 			if (!bIsCounterpartValueSame)
 			{
-				const ECheckBoxState StartingCheckedState = (bIsPropertyFilteredOut || bIsCounterpartValueSame) ? 
-					ECheckBoxState::Unchecked : InDirectParentRow.IsValid() ? InDirectParentRow.Pin()->GetWidgetCheckedState() : ECheckBoxState::Checked;
-
-				const FText DisplayName = Handle->GetPropertyDisplayName();
-
-				// Create property
-				FLevelSnapshotsEditorResultsRowPtr NewProperty = MakeShared<FLevelSnapshotsEditorResultsRow>(DisplayName, InRowType, StartingCheckedState, 
-					InResultsView, InDirectParentRow);
-
-				const TWeakPtr<FLevelSnapshotsEditorResultsRow>& ContainingObjectGroup = InDirectParentRow;
-
-				NewProperty->InitPropertyRow(ContainingObjectGroup, 
-						FoundHierarchy.IsValid() ? FoundHierarchy.Pin() : nullptr, 
-						PinnedHierarchy.IsValid() ? PinnedHierarchy : nullptr, bIsCounterpartValueSame);
-
 				// When creating the child rows, we must consider that some nodes do not list children even if the handle has children.
 				// In this case we'll create special handle hierarchies based on the handle's child handles
-				if (PinnedHierarchy->DirectChildren.Num() == 0 && ChildHandleCount > 0)
-				{
-					TSharedPtr<IPropertyHandle> CounterpartParentHandle;
-					
-					if (FoundHierarchy.IsValid() && FoundHierarchy.Pin()->Handle.IsValid())
-					{					
-						uint32 CounterpartParentHandleChildCount;
-						FoundHierarchy.Pin()->Handle->GetNumChildren(CounterpartParentHandleChildCount);
-
-						if (ensure(CounterpartParentHandleChildCount == ChildHandleCount))
-						{
-							CounterpartParentHandle = FoundHierarchy.Pin()->Handle;
-						}
-					}
-					
-					for (uint32 ChildIndex = 0; ChildIndex < ChildHandleCount; ChildIndex++)
+				if (MaxHierarchyChildCount == 0 && MaxHandleChildCount > 0)
+				{					
+					for (int32 ChildIndex = 0; ChildIndex < MaxHandleChildCount; ChildIndex++)
 					{
-						TSharedPtr<IPropertyHandle> NewChildHandle = Handle->GetChildHandle(ChildIndex);
+						TSharedPtr<IPropertyHandle> NewInputChildHandle = InputHandle->GetChildHandle(ChildIndex);
 						TSharedPtr<IPropertyHandle> NewCounterpartChildHandle;
-						bool bAreChildValuesTheSame = false;
 
-						if (CounterpartParentHandle.IsValid())
+						if (CounterpartHandle.IsValid())
 						{
-							NewCounterpartChildHandle = CounterpartParentHandle->GetChildHandle(ChildIndex);
+							NewCounterpartChildHandle = CounterpartHandle->GetChildHandle(ChildIndex);
 						}
 
-						if (NewChildHandle.IsValid())
+						if (NewInputChildHandle.IsValid() || NewCounterpartChildHandle.IsValid())
 						{
-							if (NewCounterpartChildHandle.IsValid())
-							{
-								bAreChildValuesTheSame = AreHandleValuesEqual(NewChildHandle, NewCounterpartChildHandle);
-							}
+							const bool bAreChildValuesTheSame = AreHandleValuesEqual(NewInputChildHandle, NewCounterpartChildHandle);
 
 							if (!bAreChildValuesTheSame)
 							{
-								const TSharedRef<FPropertyHandleHierarchy> NewHierarchy = MakeShared<FPropertyHandleHierarchy>(nullptr, NewChildHandle, PinnedHierarchy);
+								const TSharedRef<FPropertyHandleHierarchy> NewHierarchy = MakeShared<FPropertyHandleHierarchy>(nullptr, NewInputChildHandle, PinnedInputHierarchy);
 
-								PinnedHierarchy->DirectChildren.Add(NewHierarchy);
+								PinnedInputHierarchy->DirectChildren.Add(NewHierarchy);
 
-								if (CounterpartParentHandle.IsValid())
+								if (PinnedCounterpartHierarchy.IsValid())
 								{
 									const TSharedRef<FPropertyHandleHierarchy> NewCounterpartHierarchy = 
-										MakeShared<FPropertyHandleHierarchy>(nullptr, NewCounterpartChildHandle, FoundHierarchy.Pin());
+										MakeShared<FPropertyHandleHierarchy>(nullptr, NewCounterpartChildHandle, PinnedCounterpartHierarchy);
 
-									FoundHierarchy.Pin()->DirectChildren.Add(NewCounterpartHierarchy);
+									PinnedCounterpartHierarchy->DirectChildren.Add(NewCounterpartHierarchy);
 								}
 							}
 						}
 					}
 				}
 
-				if (PinnedHierarchy->DirectChildren.Num())
+				const ECheckBoxState StartingCheckedState = (bIsPropertyFilteredOut || bIsCounterpartValueSame) ? 
+					ECheckBoxState::Unchecked : InDirectParentRow.IsValid() ? InDirectParentRow.Pin()->GetWidgetCheckedState() : ECheckBoxState::Checked;
+
+				const FText DisplayName = InputHandle->GetPropertyDisplayName();
+
+				// Create property
+				FLevelSnapshotsEditorResultsRowPtr NewProperty = MakeShared<FLevelSnapshotsEditorResultsRow>(DisplayName, InRowType, StartingCheckedState, 
+					InResultsView, InDirectParentRow);
+
+				const TWeakPtr<FLevelSnapshotsEditorResultsRow>& ContainingObjectGroup = InDirectParentRow;
+
+				TSharedPtr<FPropertyHandleHierarchy> SnapshotHierarchy = 
+					bInputIsWorld ? (PinnedCounterpartHierarchy.IsValid() ? PinnedCounterpartHierarchy : nullptr) : (PinnedInputHierarchy.IsValid() ? PinnedInputHierarchy : nullptr);
+
+				TSharedPtr<FPropertyHandleHierarchy> WorldHierarchy = 
+					bInputIsWorld ? (PinnedInputHierarchy.IsValid() ? PinnedInputHierarchy : nullptr) : (PinnedCounterpartHierarchy.IsValid() ? PinnedCounterpartHierarchy : nullptr);
+				 
+
+				NewProperty->InitPropertyRow(ContainingObjectGroup, SnapshotHierarchy, WorldHierarchy, bIsCounterpartValueSame);
+
+				// If either hierarchy has children, then makes rows for them
+				MaxHierarchyChildCount = FMath::Max(PinnedInputHierarchy->DirectChildren.Num(), PinnedCounterpartHierarchy.IsValid() ? PinnedCounterpartHierarchy->DirectChildren.Num() : 0);
+				
+				for (int32 ChildIndex = 0; ChildIndex < MaxHierarchyChildCount; ChildIndex++)
 				{
-					for (int32 ChildIndex = 0; ChildIndex < PinnedHierarchy->DirectChildren.Num(); ChildIndex++)
+					TSharedPtr<FPropertyHandleHierarchy> ChildHierarchy;
+					
+					// Get the input Child first. If it doesn't exist at this index, check for the counterpart child.
+					const bool bIsChildFromInputHierarchy = ChildIndex < PinnedInputHierarchy->DirectChildren.Num();
+
+					if (bIsChildFromInputHierarchy)
+					{
+						ChildHierarchy = PinnedInputHierarchy->DirectChildren[ChildIndex];
+					}
+					else if (PinnedCounterpartHierarchy.IsValid() && ChildIndex < PinnedCounterpartHierarchy->DirectChildren.Num())
+					{
+						ChildHierarchy = PinnedCounterpartHierarchy->DirectChildren[ChildIndex];
+					}
+
+					if (ChildHierarchy.IsValid())
 					{
-						const TSharedRef<FPropertyHandleHierarchy>& ChildHierarchy = PinnedHierarchy->DirectChildren[ChildIndex];
+						// There may be instances of collection property members that exist on the snapshot object but not the world object
+						// For example, removing actor tags after taking a snapshot
+						ELevelSnapshotsObjectType ChildObjectInputType = ObjectType_World;
+
+						if (bInputIsWorld)
+						{
+							if (!bIsChildFromInputHierarchy)
+							{
+								ChildObjectInputType = ObjectType_Snapshot;
+							}
+						}
+						else
+						{
+							if (bIsChildFromInputHierarchy)
+							{
+								ChildObjectInputType = ObjectType_Snapshot;
+							}
+						}
 
+						const bool bShouldSearchCounterparts = SnapshotHierarchy.IsValid() && WorldHierarchy.IsValid();
+						
 						LoopOverHandleHierarchiesAndCreateRowHierarchy(
-							ChildHierarchy, NewProperty, PropertiesThatPassFilter, PropertyRowsGenerated, InResultsView, InHierarchyToSearchForCounterparts);
+							ChildObjectInputType, ChildHierarchy, NewProperty, PropertiesThatPassFilter, PropertyRowsGenerated, InResultsView, bShouldSearchCounterparts ? (bIsChildFromInputHierarchy ? InHierarchyToSearchForCounterparts : InputHierarchy) : nullptr);
 					}
 				}
 
-				if (NewProperty->DoesRowRepresentGroup() && !NewProperty->GetChildRows().Num() && !Handle->IsCustomized())
+				if (NewProperty->DoesRowRepresentGroup() && !NewProperty->GetChildRows().Num() && !InputHandle->IsCustomized())
 				{
 					// No valid children, destroy group
 					NewProperty.Reset();
@@ -1469,7 +1541,7 @@ void FLevelSnapshotsEditorResultsRow::SetWidgetCheckedState(const ECheckBoxState
 				continue;
 			}
 
-			ChildRow->SetWidgetCheckedState(NewState);
+			ChildRow->SetWidgetCheckedState(NewState, true);
 		}
 		
 		EvaluateAndSetAllParentGroupCheckedStates();
@@ -2076,10 +2148,9 @@ FReply SLevelSnapshotsEditorResults::OnClickApplyToWorld()
 	const TOptional<ULevelSnapshot*> ActiveLevelSnapshot = EditorDataPtr->GetActiveSnapshot();
 	if (ActiveLevelSnapshot.IsSet())
 	{
-		DECLARE_SCOPE_CYCLE_COUNTER(TEXT("OnClickApplyToWorld"), STAT_LevelSnapshots, STATGROUP_LevelSnapshots);
+		SCOPED_SNAPSHOT_EDITOR_TRACE(ClickApplyToWorld);
 		{
-			// Measure how long it takes to get all selected properties from UI
-			DECLARE_SCOPE_CYCLE_COUNTER(TEXT("BuildSelectionSetFromSelectedPropertiesInEachActorGroup"), STAT_BuildSelectionSetFromSelectedProperties, STATGROUP_LevelSnapshots);
+			SCOPED_SNAPSHOT_EDITOR_TRACE(GetSelectedPropertiesFromUI)
 			BuildSelectionSetFromSelectedPropertiesInEachActorGroup();
 		}
 
@@ -2482,7 +2553,7 @@ void SLevelSnapshotsEditorResults::GenerateTreeView(const bool bSnapshotHasChang
 		return;
 	}
 
-	DECLARE_SCOPE_CYCLE_COUNTER(TEXT("GenerateTreeView"), STAT_GenerateTreeView, STATGROUP_LevelSnapshots);
+	SCOPED_SNAPSHOT_EDITOR_TRACE(GenerateTreeView);
 
 	if (bSnapshotHasChanged)
 	{
@@ -2645,12 +2716,6 @@ bool SLevelSnapshotsEditorResults::GenerateTreeViewChildren_ModifiedActors(FLeve
 					{
 						if (const FPropertySelection* SubObjectProperties = InSelectionMap.GetSelectedProperties(SubObject))
 						{
-							if (!SubObject->IsA(UActorComponent::StaticClass()))
-							{
-								FPropertySelection* InObjectProperties = const_cast<FPropertySelection*>(InSelectionMap.GetSelectedProperties(InObject));
-								InObjectProperties->SetHasCustomSerializedSubobjects(true);
-							}
-							
 							NewCachedSearchTerms += " " + SubObject->GetName();
 							for (const TFieldPath<FProperty>& LeafProperty : SubObjectProperties->GetSelectedLeafProperties())
 							{
@@ -3010,7 +3075,6 @@ void SLevelSnapshotsEditorResultsRow::Construct(const FArguments& InArgs, const
 	BasicRowWidgets->AddSlot()
 	.VAlign(VAlign_Center)
 	.HAlign(HAlign_Left)
-	.AutoWidth()
 	[
 		SNew(STextBlock).Text(DisplayText)
 	];
@@ -3019,20 +3083,30 @@ void SLevelSnapshotsEditorResultsRow::Construct(const FArguments& InArgs, const
 	
 	if (bDoesRowNeedSplitter)
 	{
-		SAssignNew(SplitterPtr, SSplitter)
+		SAssignNew(OuterSplitterPtr, SSplitter)
 		.Style(FEditorStyle::Get(), "DetailsView.Splitter")
-		.PhysicalSplitterHandleSize(bDoesRowNeedSplitter ? 5.0f : 0.f)
-		.HitDetectionSplitterHandleSize(bDoesRowNeedSplitter ? 5.0f : 0.f);
+		.PhysicalSplitterHandleSize(5.0f)
+		.HitDetectionSplitterHandleSize(5.0f);
 
-		SplitterPtr->AddSlot()
-		.OnSlotResized(SSplitter::FOnSlotResized::CreateSP(this, &SLevelSnapshotsEditorResultsRow::SetNameColumnSize))
+		OuterSplitterPtr->AddSlot().SizeRule(SSplitter::ESizeRule::FractionOfParent)
+		.OnSlotResized(SSplitter::FOnSlotResized::CreateLambda([this] (float InWidth) {}))
 		.Value(TAttribute<float>::Create(TAttribute<float>::FGetter::CreateSP(this, &SLevelSnapshotsEditorResultsRow::GetNameColumnSize)))
 		[
 			BasicRowWidgets
 		];
+
+		OuterSplitterPtr->AddSlot()
+		.OnSlotResized(SSplitter::FOnSlotResized::CreateSP(this, &SLevelSnapshotsEditorResultsRow::SetNestedColumnSize))
+		.Value(TAttribute<float>::Create(TAttribute<float>::FGetter::CreateSP(this, &SLevelSnapshotsEditorResultsRow::CalculateAndReturnNestedColumnSize)))
+		[
+			SAssignNew(NestedSplitterPtr, SSplitter)
+			.Style(FEditorStyle::Get(), "DetailsView.Splitter")
+			.PhysicalSplitterHandleSize(5.0f)
+			.HitDetectionSplitterHandleSize(5.0f)
+		];
 	
-		// Splitter Slot 1
-		int32 SlotIndex = 1;
+		// Nested Splitter Slot 0
+		int32 SlotIndex = 0;
 		
 		TSharedPtr<SWidget> WorldChildWidget;
 
@@ -3071,9 +3145,9 @@ void SLevelSnapshotsEditorResultsRow::Construct(const FArguments& InArgs, const
 		}
 		else
 		{
-			if (bIsHeaderRow && PinnedItem->GetHeaderColumns().Num() > SlotIndex)
+			if (bIsHeaderRow && PinnedItem->GetHeaderColumns().Num() > SlotIndex + 1)
 			{
-				WorldChildWidget = SNew(STextBlock).Text(PinnedItem->GetHeaderColumns()[SlotIndex]);
+				WorldChildWidget = SNew(STextBlock).Text(PinnedItem->GetHeaderColumns()[SlotIndex + 1]);
 			}
 			else if (bIsAddedOrRemovedActorRow || PinnedItem->DoesRowRepresentGroup())
 			{
@@ -3106,15 +3180,15 @@ void SLevelSnapshotsEditorResultsRow::Construct(const FArguments& InArgs, const
 				];
 		}
 
-		SplitterPtr->AddSlot()
-		.OnSlotResized(SSplitter::FOnSlotResized::CreateSP(this, &SLevelSnapshotsEditorResultsRow::SetWorldColumnSize))
+		NestedSplitterPtr->AddSlot()
+			.OnSlotResized(SSplitter::FOnSlotResized::CreateLambda([this] (float InWidth) {}))
 		.Value(TAttribute<float>::Create(TAttribute<float>::FGetter::CreateSP(this, &SLevelSnapshotsEditorResultsRow::GetWorldColumnSize)))
 		[
 			FinalWorldWidget.ToSharedRef()
 		];
 
-		// Splitter Slot 2
-		SlotIndex = 2;
+		// Splitter Slot 1
+		SlotIndex = 1;
 		
 		TSharedPtr<SWidget> SnapshotChildWidget;
 		
@@ -3153,9 +3227,9 @@ void SLevelSnapshotsEditorResultsRow::Construct(const FArguments& InArgs, const
 		}
 		else
 		{
-			if (bIsHeaderRow && PinnedItem->GetHeaderColumns().Num() > SlotIndex)
+			if (bIsHeaderRow && PinnedItem->GetHeaderColumns().Num() > SlotIndex + 1)
 			{
-				SnapshotChildWidget = SNew(STextBlock).Text(PinnedItem->GetHeaderColumns()[SlotIndex]);
+				SnapshotChildWidget = SNew(STextBlock).Text(PinnedItem->GetHeaderColumns()[SlotIndex + 1]);
 			}
 			else if (bIsAddedOrRemovedActorRow || PinnedItem->DoesRowRepresentGroup())
 			{
@@ -3188,14 +3262,14 @@ void SLevelSnapshotsEditorResultsRow::Construct(const FArguments& InArgs, const
 				];
 		}
 
-		SplitterPtr->AddSlot()
+		NestedSplitterPtr->AddSlot()
 		.OnSlotResized(SSplitter::FOnSlotResized::CreateSP(this, &SLevelSnapshotsEditorResultsRow::SetSnapshotColumnSize))
 		.Value(TAttribute<float>::Create(TAttribute<float>::FGetter::CreateSP(this, &SLevelSnapshotsEditorResultsRow::GetSnapshotColumnSize)))
 		[
 			FinalSnapshotWidget.ToSharedRef()
 		];
 		
-		BorderPtr->SetContent(SplitterPtr.ToSharedRef());
+		BorderPtr->SetContent(OuterSplitterPtr.ToSharedRef());
 	}
 	else
 	{
@@ -3208,36 +3282,69 @@ SLevelSnapshotsEditorResultsRow::~SLevelSnapshotsEditorResultsRow()
 	// Remove delegate bindings
 
 	// Unbind event to the splitter being resized first
-	if (SplitterPtr.IsValid())
+	if (NestedSplitterPtr.IsValid())
+	{
+		for (int32 SplitterSlotCount = 0; SplitterSlotCount < NestedSplitterPtr->GetChildren()->Num(); SplitterSlotCount++)
+		{
+			NestedSplitterPtr->SlotAt(SplitterSlotCount).OnSlotResized();
+		}
+	}
+
+	if (OuterSplitterPtr.IsValid())
 	{
-		for (int32 SplitterSlotCount = 0; SplitterSlotCount < SplitterPtr->GetChildren()->Num(); SplitterSlotCount++)
+		for (int32 SplitterSlotCount = 0; SplitterSlotCount < OuterSplitterPtr->GetChildren()->Num(); SplitterSlotCount++)
 		{
-			SplitterPtr->SlotAt(SplitterSlotCount).OnSlotResized().Unbind();
+			OuterSplitterPtr->SlotAt(SplitterSlotCount).OnSlotResized();
 		}
 	}
 
-	SplitterPtr.Reset();
+	OuterSplitterPtr.Reset();
+	NestedSplitterPtr.Reset();
 	SplitterManagerPtr.Reset();
 }
 
 float SLevelSnapshotsEditorResultsRow::GetNameColumnSize() const
 {
-	return SplitterManagerPtr->NameColumnWidth;;
+	return 1.0f - CachedNestedColumnWidthAdjusted;
+}
+
+float SLevelSnapshotsEditorResultsRow::CalculateAndReturnNestedColumnSize()
+{
+	check(Item.IsValid());
+	
+	const TSharedPtr<FLevelSnapshotsEditorResultsRow>& PinnedItem = Item.Pin();
+	const uint8 ChildDepth = PinnedItem->GetChildDepth();
+	const float StartWidth = SplitterManagerPtr->NestedColumnWidth;
+
+	if (ChildDepth > 0)
+	{
+		const float LocalPixelOffset = 10.0f;
+		const float LocalPixelDifference = LocalPixelOffset * ChildDepth;
+		const float LocalSizeX = GetTickSpaceGeometry().GetLocalSize().X;
+		const float NestedItemCoefficient = (LocalSizeX + LocalPixelDifference) / LocalSizeX;
+		
+		CachedNestedColumnWidthAdjusted = StartWidth * NestedItemCoefficient;
+		
+		return CachedNestedColumnWidthAdjusted;
+	}
+
+	return CachedNestedColumnWidthAdjusted = StartWidth;
 }
 
 float SLevelSnapshotsEditorResultsRow::GetSnapshotColumnSize() const
 {
-	return SplitterManagerPtr->SnapshotPropertyColumnWidth;
+	const float EndWidth = SplitterManagerPtr->SnapshotPropertyColumnWidth;
+	return EndWidth;
 }
 
 float SLevelSnapshotsEditorResultsRow::GetWorldColumnSize() const
 {
-	return SplitterManagerPtr->WorldObjectPropertyColumnWidth;
+	return 1.0f - GetSnapshotColumnSize();
 }
 
-void SLevelSnapshotsEditorResultsRow::SetNameColumnSize(const float InWidth) const
+void SLevelSnapshotsEditorResultsRow::SetNestedColumnSize(const float InWidth) const
 {
-	SplitterManagerPtr->NameColumnWidth = InWidth;
+	SplitterManagerPtr->NestedColumnWidth = InWidth;
 }
 
 void SLevelSnapshotsEditorResultsRow::SetSnapshotColumnSize(const float InWidth) const
@@ -3245,15 +3352,10 @@ void SLevelSnapshotsEditorResultsRow::SetSnapshotColumnSize(const float InWidth)
 	SplitterManagerPtr->SnapshotPropertyColumnWidth = InWidth;
 }
 
-void SLevelSnapshotsEditorResultsRow::SetWorldColumnSize(const float InWidth) const
-{
-	SplitterManagerPtr->WorldObjectPropertyColumnWidth = InWidth;
-}
-
 FReply SLevelSnapshotsEditorResultsRow::OnMouseButtonUp(const FGeometry& MyGeometry, const FPointerEvent& MouseEvent)
 {
 	if (Item.IsValid() && Item.Pin()->GetRowType() == FLevelSnapshotsEditorResultsRow::ActorGroup)
-	{
+	{		
 		if (AActor* Actor = Cast<AActor>(Item.Pin()->GetWorldObject()))
 		{
 			const FLevelSnapshotsEditorModule& Module = FLevelSnapshotsEditorModule::Get();
diff --git a/Engine/Plugins/VirtualProduction/LevelSnapshots/Source/LevelSnapshotsEditor/Private/Views/Results/SLevelSnapshotsEditorResults.h b/Engine/Plugins/VirtualProduction/LevelSnapshots/Source/LevelSnapshotsEditor/Private/Views/Results/SLevelSnapshotsEditorResults.h
--- a/Engine/Plugins/VirtualProduction/LevelSnapshots/Source/LevelSnapshotsEditor/Private/Views/Results/SLevelSnapshotsEditorResults.h
+++ b/Engine/Plugins/VirtualProduction/LevelSnapshots/Source/LevelSnapshotsEditor/Private/Views/Results/SLevelSnapshotsEditorResults.h
@@ -45,9 +45,8 @@ enum ELevelSnapshotsObjectType
 
 struct FLevelSnapshotsEditorResultsSplitterManager
 {
-	float NameColumnWidth = 1.0f;
-	float SnapshotPropertyColumnWidth = 1.0f;
-	float WorldObjectPropertyColumnWidth = 1.0f;
+	float NestedColumnWidth = 0.5f; // The right side of the first splitter which contains the nested splitter for the property widgets
+	float SnapshotPropertyColumnWidth = 0.5f;
 };
 
 struct FLevelSnapshotsEditorResultsRowStateMemory
@@ -482,23 +481,30 @@ public:
 
 	float GetNameColumnSize() const;
 
+	float CalculateAndReturnNestedColumnSize();
+
 	float GetSnapshotColumnSize() const;
 
 	float GetWorldColumnSize() const;
 
-	void SetNameColumnSize(const float InWidth) const;
+	void SetNestedColumnSize(const float InWidth) const;
 
 	void SetSnapshotColumnSize(const float InWidth) const;
 
-	void SetWorldColumnSize(const float InWidth) const;
-
 	virtual FReply OnMouseButtonUp(const FGeometry& MyGeometry, const FPointerEvent& MouseEvent) override;
 
 private:
 	
-	TSharedPtr<SSplitter> SplitterPtr = nullptr;
+	TSharedPtr<SSplitter> OuterSplitterPtr = nullptr;
+	TSharedPtr<SSplitter> NestedSplitterPtr = nullptr;
+	
 	TWeakPtr<FLevelSnapshotsEditorResultsRow> Item = nullptr;
 
 	/* For splitter sync */
+
+	/* To sync up splitter location in tree view items, we need to account for the tree view's indentation.
+	 * Instead of calculating the coefficient twice each frame (for left and right splitter slots), we do it once and cache it here. */
+	float CachedNestedColumnWidthAdjusted = 0.f;
+	
 	FLevelSnapshotsEditorResultsSplitterManagerPtr SplitterManagerPtr;
 };
diff --git a/Engine/Plugins/VirtualProduction/LevelSnapshots/Source/LevelSnapshotsEditor/Private/Widgets/SLevelSnapshotsEditorBrowser.cpp b/Engine/Plugins/VirtualProduction/LevelSnapshots/Source/LevelSnapshotsEditor/Private/Widgets/SLevelSnapshotsEditorBrowser.cpp
--- a/Engine/Plugins/VirtualProduction/LevelSnapshots/Source/LevelSnapshotsEditor/Private/Widgets/SLevelSnapshotsEditorBrowser.cpp
+++ b/Engine/Plugins/VirtualProduction/LevelSnapshots/Source/LevelSnapshotsEditor/Private/Widgets/SLevelSnapshotsEditorBrowser.cpp
@@ -8,10 +8,9 @@
 
 #include "IContentBrowserSingleton.h"
 #include "ContentBrowserModule.h"
-#include "LevelSnapshotsEditorStyle.h"
+#include "LevelSnapshotsLog.h"
 #include "AssetRegistry/AssetRegistryModule.h"
 #include "Editor/EditorStyle/Public/EditorStyleSet.h"
-#include "Framework/Commands/GenericCommands.h"
 #include "Misc/ScopedSlowTask.h"
 #include "Modules/ModuleManager.h"
 #include "Slate/Public/Framework/MultiBox/MultiBoxBuilder.h"
@@ -58,6 +57,8 @@ void SLevelSnapshotsEditorBrowser::Construct(const FArguments& InArgs, const TSh
 
 void SLevelSnapshotsEditorBrowser::SelectAsset(const FAssetData& InAssetData) const
 {
+	SCOPED_SNAPSHOT_EDITOR_TRACE(SelectSnapshotAsset);
+	
 	FScopedSlowTask SelectSnapshot(100.f, LOCTEXT("SelectSnapshotKey", "Loading snapshot"));
 	SelectSnapshot.EnterProgressFrame(60.f);
 	SelectSnapshot.MakeDialog();
diff --git a/Engine/Plugins/VirtualProduction/LevelSnapshots/Source/LevelSnapshotsEditor/Private/Widgets/SLevelSnapshotsEditorCreationForm.cpp b/Engine/Plugins/VirtualProduction/LevelSnapshots/Source/LevelSnapshotsEditor/Private/Widgets/SLevelSnapshotsEditorCreationForm.cpp
--- a/Engine/Plugins/VirtualProduction/LevelSnapshots/Source/LevelSnapshotsEditor/Private/Widgets/SLevelSnapshotsEditorCreationForm.cpp
+++ b/Engine/Plugins/VirtualProduction/LevelSnapshots/Source/LevelSnapshotsEditor/Private/Widgets/SLevelSnapshotsEditorCreationForm.cpp
@@ -15,18 +15,20 @@
 #include "ISettingsSection.h"
 #include "Modules/ModuleManager.h"
 #include "PropertyEditorModule.h"
+#include "Widgets/Images/SImage.h"
 #include "Widgets/Input/SButton.h"
+#include "Widgets/Input/SCheckBox.h"
 #include "Widgets/Input/SEditableTextBox.h"
-#include "Widgets/Images/SImage.h"
 #include "Widgets/Input/SMultiLineEditableTextBox.h"
 #include "Widgets/SWindow.h"
 #include "Widgets/Text/STextBlock.h"
 
-TSharedRef<SWindow> FLevelSnapshotsEditorCreationForm::MakeAndShowCreationWindow(
+TSharedRef<SWindow> SLevelSnapshotsEditorCreationForm::MakeAndShowCreationWindow(
 	const FCloseCreationFormDelegate& CallOnClose, 
 	ULevelSnapshotsEditorProjectSettings* InProjectSettings, ULevelSnapshotsEditorDataManagementSettings* InDataManagementSettings)
 {
 	check(InProjectSettings);
+	check(InDataManagementSettings);
 	
 	// Compute centered window position based on max window size, which include when all categories are expanded
 	const FVector2D LastSize = InProjectSettings->GetLastCreationWindowSize();
@@ -49,27 +51,27 @@ TSharedRef<SWindow> FLevelSnapshotsEditorCreationForm::MakeAndShowCreationWindow
 		.SupportsMaximize(false)
 		.ScreenPosition(WindowPosition);
 
-	const TSharedRef<SLevelSnapshotsEditorCreationForm> CreationForm = 
-		SNew(SLevelSnapshotsEditorCreationForm, Window, CallOnClose, InProjectSettings, InDataManagementSettings);
-	
+	const TSharedRef<SLevelSnapshotsEditorCreationForm> CreationForm = SNew(SLevelSnapshotsEditorCreationForm, Window, CallOnClose, InProjectSettings, InDataManagementSettings);
 	Window->SetContent
 	(
 		CreationForm
 	);
-
 	Window->SetOnWindowClosed(FOnWindowClosed::CreateSP(CreationForm, &SLevelSnapshotsEditorCreationForm::OnWindowClosed));
 
 	FSlateApplication::Get().AddWindow(Window);
-
 	return Window;
 }
 
 void SLevelSnapshotsEditorCreationForm::Construct(
-	const FArguments& InArgs, TWeakPtr< SWindow > InWidgetWindow, const FCloseCreationFormDelegate& CallOnClose, 
-	ULevelSnapshotsEditorProjectSettings* InProjectSettings, ULevelSnapshotsEditorDataManagementSettings* InDataManagementSettings)
+	const FArguments& InArgs,
+	TWeakPtr<SWindow> InWidgetWindow,
+	const FCloseCreationFormDelegate& CallOnClose, 
+	ULevelSnapshotsEditorProjectSettings* InProjectSettings,
+	ULevelSnapshotsEditorDataManagementSettings* InDataManagementSettings)
 {
 	check(InProjectSettings);
-
+	check(InDataManagementSettings);
+	
 	ProjectSettingsObjectPtr = InProjectSettings;
 	DataManagementSettingsObjectPtr = InDataManagementSettings;
 	
@@ -79,117 +81,141 @@ void SLevelSnapshotsEditorCreationForm::Construct(
 	bNameDiffersFromDefault = DataManagementSettingsObjectPtr.Get()->IsNameOverridden();
 	
 	ChildSlot
+	[
+		SNew(SBorder)
+		.BorderImage(new FSlateColorBrush(FColor(10, 10, 10)))
 		[
-			SNew(SBorder)
-			.BorderImage(new FSlateColorBrush(FColor(10, 10, 10)))
-			[
-				SNew(SVerticalBox)
+			SNew(SVerticalBox)
 
-				+ SVerticalBox::Slot()
-				.AutoHeight()
-				.Padding(8.f, 5.f, 8.f, 0.f)
-				.HAlign(HAlign_Fill)
-				.VAlign(VAlign_Top)
+			+ SVerticalBox::Slot()
+			.AutoHeight()
+			.Padding(2.f, 2.f, 2.f, 0.f)
+			.HAlign(HAlign_Fill)
+			.VAlign(VAlign_Top)
+			[
+				SNew(SBorder)
+				.BorderImage(FEditorStyle::GetBrush("ToolPanel.DarkGroupBorder"))
 				[
-					SNew(SBorder)
-					.BorderImage(FEditorStyle::GetBrush("ToolPanel.DarkGroupBorder"))
+					SNew(SVerticalBox)
+
+					+SVerticalBox::Slot()
+					.AutoHeight()
+					.Padding(1.f, 1.f, 0.f, 0.f)
+					.HAlign(HAlign_Left)
+					.VAlign(VAlign_Top)
+					[
+						SNew(STextBlock)
+						.Font(FCoreStyle::GetDefaultFontStyle("Regular", 9))
+						.ColorAndOpacity(FColor(200, 200, 200))
+						.Text(NSLOCTEXT("LevelSnapshots", "CreationForm_SnapshotNameLabel", "Name"))
+					]
+
+					+SVerticalBox::Slot()
+					.AutoHeight()
+					.Padding(8.f, 1.f, 8.f, 10.f)
 					[
-						SNew(SVerticalBox)
+						SNew(SHorizontalBox)
 
-						+SVerticalBox::Slot()
-						.AutoHeight()
-						.Padding(1.f, 1.f, 0.f, 0.f)
-						.HAlign(HAlign_Left)
-						.VAlign(VAlign_Top)
+						+ SHorizontalBox::Slot()
+						.HAlign(HAlign_Fill)
 						[
-							SNew(STextBlock)
-							.Font(FCoreStyle::GetDefaultFontStyle("Regular", 9))
-							.ColorAndOpacity(FColor(200, 200, 200))
-							.Text(NSLOCTEXT("LevelSnapshots", "CreationForm_SnapshotNameLabel", "NAME"))
+							SNew(SEditableTextBox)
+							.Font(FCoreStyle::GetDefaultFontStyle("Bold", 16))
+							.BackgroundColor(FLinearColor::Transparent)
+							.ForegroundColor(FSlateColor::UseForeground())
+							.Justification(ETextJustify::Center)
+							.SelectAllTextWhenFocused(true)
+							.HintText(NSLOCTEXT("LevelSnapshots", "CreationForm_SnapshotNameOverrideHintText", "Override Snapshot Name..."))
+							.Text(this, &SLevelSnapshotsEditorCreationForm::GetNameOverrideText)
+							.OnTextCommitted(this, &SLevelSnapshotsEditorCreationForm::SetNameOverrideText)
+							.ToolTipText(
+							NSLOCTEXT("LevelSnapshots", "CreationForm_NameOverrideFieldTooltipText", "Override the name defined in Project Settings while using the Creation Form."))
 						]
 
-						+SVerticalBox::Slot()
-						.AutoHeight()
-						.Padding(8.f, 1.f, 8.f, 10.f)
+						+SHorizontalBox::Slot()
+						.AutoWidth()
+						.HAlign(HAlign_Right)
+						.VAlign(VAlign_Center)
 						[
-							SNew(SHorizontalBox)
-
-							+ SHorizontalBox::Slot()
-							.HAlign(HAlign_Fill)
+							SNew(SButton)
+							.IsFocusable(false)
+							.ToolTipText(
+								NSLOCTEXT("LevelSnapshots", "CreationForm_ResetNameTooltipText", "Reset the overridden name to the one defined in Project Settings."))
+							.ButtonStyle(FEditorStyle::Get(), "NoBorder")
+							.ContentPadding(0)
+							.Visibility(this, &SLevelSnapshotsEditorCreationForm::GetNameDiffersFromDefaultAsVisibility)
+							.OnClicked(this, &SLevelSnapshotsEditorCreationForm::OnResetNameClicked)
+							.Content()
 							[
-								SNew(SEditableTextBox)
-								.Font(FCoreStyle::GetDefaultFontStyle("Bold", 16))
-								.BackgroundColor(FLinearColor::Transparent)
-								.ForegroundColor(FSlateColor::UseForeground())
-								.Justification(ETextJustify::Center)
-								.SelectAllTextWhenFocused(true)
-								.HintText(NSLOCTEXT("LevelSnapshots", "CreationForm_SnapshotNameOverrideHintText", "Override Snapshot Name..."))
-								.Text(this, &SLevelSnapshotsEditorCreationForm::GetNameOverrideText)
-								.OnTextCommitted(this, &SLevelSnapshotsEditorCreationForm::SetNameOverrideText)
-								.ToolTipText(
-								NSLOCTEXT("LevelSnapshots", "CreationForm_NameOverrideFieldTooltipText", "Override the name defined in Project Settings while using the Creation Form."))
-							]
-
-							+SHorizontalBox::Slot()
-							.AutoWidth()
-							.HAlign(HAlign_Right)
-							.VAlign(VAlign_Center)
-							[
-								SNew(SButton)
-								.IsFocusable(false)
-								.ToolTipText(
-									NSLOCTEXT("LevelSnapshots", "CreationForm_ResetNameTooltipText", "Reset the overridden name to the one defined in Project Settings."))
-								.ButtonStyle(FEditorStyle::Get(), "NoBorder")
-								.ContentPadding(0)
-								.Visibility(this, &SLevelSnapshotsEditorCreationForm::GetNameDiffersFromDefaultAsVisibility)
-								.OnClicked(this, &SLevelSnapshotsEditorCreationForm::OnResetNameClicked)
-								.Content()
-								[
-									SNew(SImage)
-									.Image(FEditorStyle::GetBrush("PropertyWindow.DiffersFromDefault"))
-								]
+								SNew(SImage)
+								.Image(FEditorStyle::GetBrush("PropertyWindow.DiffersFromDefault"))
 							]
 						]
 					]
 				]
+			]
 
-				+SVerticalBox::Slot()
-				.Padding(8.f, 10.f, 8.f, 20.f)
-				.VAlign(VAlign_Fill)
-				[
-					SNew(SMultiLineEditableTextBox)
-					.Font(FCoreStyle::GetDefaultFontStyle("Bold", 10))
-					.BackgroundColor(FLinearColor(0.2f, 0.2f, 0.2f))
-					.ForegroundColor(FSlateColor::UseForeground())
-					.SelectAllTextWhenFocused(true)
-					.HintText(NSLOCTEXT("LevelSnapshots", "CreationForm_DescriptionHintText", "<description>"))
-					.Text(DescriptionText)
-					.OnTextCommitted(this, &SLevelSnapshotsEditorCreationForm::SetDescriptionText)
-					.AllowMultiLine(true)
-					.AutoWrapText(true)
-				]
-				
-				+SVerticalBox::Slot()
-				.AutoHeight()
-				.VAlign(VAlign_Bottom)
-				[
-					SNew(STextBlock)
-					.TextStyle(FEditorStyle::Get(), "NormalText.Important")
-					.Text(NSLOCTEXT("LevelSnapshots", "CreationForm_SaveDirLabel", "Save Directory"))
-				]
+			+SVerticalBox::Slot()
+			.Padding(2.f, 10.f, 2.f, 0.f)
+			.VAlign(VAlign_Fill)
+			[
+				SNew(SMultiLineEditableTextBox)
+				.Font(FCoreStyle::GetDefaultFontStyle("Bold", 10))
+				.BackgroundColor(FLinearColor(0.2f, 0.2f, 0.2f))
+				.ForegroundColor(FSlateColor::UseForeground())
+				.SelectAllTextWhenFocused(true)
+				.HintText(NSLOCTEXT("LevelSnapshots", "CreationForm_DescriptionHintText", "Description"))
+				.Text(DescriptionText)
+				.OnTextCommitted(this, &SLevelSnapshotsEditorCreationForm::SetDescriptionText)
+				.AllowMultiLine(true)
+				.AutoWrapText(true)
+			]
+			
+			+SVerticalBox::Slot()
+			.Padding(2.f, 10.f, 2.f, 0.f)
+			.AutoHeight()
+			.VAlign(VAlign_Bottom)
+			[
+				SNew(STextBlock)
+				.TextStyle(FEditorStyle::Get(), "NormalText.Important")
+				.Text(NSLOCTEXT("LevelSnapshots", "CreationForm_SaveDirLabel", "Save Directory"))
+			]
+
+			+SVerticalBox::Slot()
+			.Padding(2.f, 2.f, 2.f, 0.f)
+			.AutoHeight()
+			.VAlign(VAlign_Bottom)
+			[
+				MakeDataManagementSettingsDetailsWidget()
+			]
+			
+			+SVerticalBox::Slot()
+			.AutoHeight()
+			.VAlign(VAlign_Bottom)
+			.HAlign(HAlign_Fill)
+			.Padding(2.f, 5.f)
+			[
+				SNew(SHorizontalBox)
 
-				+SVerticalBox::Slot()
-				.AutoHeight()
-				.VAlign(VAlign_Bottom)
+				// Save Async checkbox
+				+SHorizontalBox::Slot()
+				.Padding(3.f, 0.f)
+				.HAlign(HAlign_Left)
 				[
-					MakeDataManagementSettingsDetailsWidget()
+					SNew(SCheckBox)
+					.IsChecked_Lambda([this]() { return bSaveAsync ? ECheckBoxState::Checked : ECheckBoxState::Unchecked; })
+					.OnCheckStateChanged_Lambda([this](ECheckBoxState NewState) { bSaveAsync = NewState == ECheckBoxState::Checked; })
+					.ToolTipText(NSLOCTEXT("LevelSnapshots", "CreationForm_SaveAsync_Tooltip", "Enabling may speed up saving for large levels."))
+					[
+						SNew(STextBlock)
+						.Font(FCoreStyle::GetDefaultFontStyle("Regular", 10))
+						.Text(NSLOCTEXT("LevelSnapshots", "CreationForm_SaveAsync", "Save async"))
+					]
 				]
-
-				+SVerticalBox::Slot()
-				.AutoHeight()
-				.VAlign(VAlign_Bottom)
+				
+				// Create snapshot button
+				+SHorizontalBox::Slot()
 				.HAlign(HAlign_Right)
-				.Padding(2.f, 5.f)
 				[
 					SNew(SButton)
 					.ButtonStyle(FEditorStyle::Get(), "FlatButton.Success")
@@ -203,7 +229,8 @@ void SLevelSnapshotsEditorCreationForm::Construct(
 					]
 				]
 			]
-		];
+		]
+	];
 }
 
 TSharedRef<SWidget> SLevelSnapshotsEditorCreationForm::MakeDataManagementSettingsDetailsWidget() const
@@ -301,6 +328,9 @@ void SLevelSnapshotsEditorCreationForm::OnWindowClosed(const TSharedRef<SWindow>
 	{
 		DataManagementSettingsObjectPtr->SaveConfig();
 	}
-	
-	CallOnCloseDelegate.ExecuteIfBound(bWasCreateSnapshotPressed, DescriptionText);
+
+	if (bWasCreateSnapshotPressed)
+	{
+		CallOnCloseDelegate.ExecuteIfBound(DescriptionText, bSaveAsync);
+	}
 }
diff --git a/Engine/Plugins/VirtualProduction/LevelSnapshots/Source/LevelSnapshotsEditor/Private/Widgets/SLevelSnapshotsEditorCreationForm.h b/Engine/Plugins/VirtualProduction/LevelSnapshots/Source/LevelSnapshotsEditor/Private/Widgets/SLevelSnapshotsEditorCreationForm.h
--- a/Engine/Plugins/VirtualProduction/LevelSnapshots/Source/LevelSnapshotsEditor/Private/Widgets/SLevelSnapshotsEditorCreationForm.h
+++ b/Engine/Plugins/VirtualProduction/LevelSnapshots/Source/LevelSnapshotsEditor/Private/Widgets/SLevelSnapshotsEditorCreationForm.h
@@ -7,30 +7,29 @@
 #include "Settings/LevelSnapshotsEditorDataManagementSettings.h"
 #include "Dialogs/CustomDialog.h"
 
-DECLARE_DELEGATE_TwoParams(FCloseCreationFormDelegate, bool, FText)
+DECLARE_DELEGATE_TwoParams(FCloseCreationFormDelegate, const FText& /* Description */, bool /* bSaveAsync */);
 
 class ULevelSnapshotsEditorProjectSettings;
 class ULevelSnapshotsEditorDataManagementSettings;
 class SWindow;
 
-struct FLevelSnapshotsEditorCreationForm
-{
-	static TSharedRef<SWindow> MakeAndShowCreationWindow(
-		const FCloseCreationFormDelegate& CallOnClose, 
-		ULevelSnapshotsEditorProjectSettings* InProjectSettings, ULevelSnapshotsEditorDataManagementSettings* InDataManagementSettings);
-};
-
 class SLevelSnapshotsEditorCreationForm : public SCustomDialog
 {
 public:
 
+	static TSharedRef<SWindow> MakeAndShowCreationWindow(const FCloseCreationFormDelegate& CallOnClose, ULevelSnapshotsEditorProjectSettings* InProjectSettings, ULevelSnapshotsEditorDataManagementSettings* InDataManagementSettings);
+	
 	SLATE_BEGIN_ARGS(SLevelSnapshotsEditorCreationForm)
 	{}
 	SLATE_END_ARGS()
 
 	void Construct(
-		const FArguments& InArgs, TWeakPtr< SWindow > InWidgetWindow, const FCloseCreationFormDelegate& CallOnClose, 
-		ULevelSnapshotsEditorProjectSettings* InProjectSettings, ULevelSnapshotsEditorDataManagementSettings* InDataManagementSettings);
+		const FArguments& InArgs,
+		TWeakPtr<SWindow> InWidgetWindow,
+		const FCloseCreationFormDelegate& CallOnClose,
+		ULevelSnapshotsEditorProjectSettings* InProjectSettings,
+		ULevelSnapshotsEditorDataManagementSettings* InDataManagementSettings
+		);
 
 	~SLevelSnapshotsEditorCreationForm();
 
@@ -60,7 +59,7 @@ private:
 	bool bWasCreateSnapshotPressed = false;
 
 	FText DescriptionText;
+	bool bSaveAsync = true;
 
 	FCloseCreationFormDelegate CallOnCloseDelegate;
-
 };
diff --git a/Engine/Plugins/VirtualProduction/Takes/Source/TakeRecorderSources/Private/TakeRecorderNearbySpawnedActorSource.cpp b/Engine/Plugins/VirtualProduction/Takes/Source/TakeRecorderSources/Private/TakeRecorderNearbySpawnedActorSource.cpp
--- a/Engine/Plugins/VirtualProduction/Takes/Source/TakeRecorderSources/Private/TakeRecorderNearbySpawnedActorSource.cpp
+++ b/Engine/Plugins/VirtualProduction/Takes/Source/TakeRecorderSources/Private/TakeRecorderNearbySpawnedActorSource.cpp
@@ -176,7 +176,7 @@ void UTakeRecorderNearbySpawnedActorSource::HandleActorSpawned(AActor* Actor, cl
 		return;
 	}
 
-	UE_LOG(LogTakesCore, Log, TEXT("Actor: %s PendingKill: %d PendingKillOrUnreachable: %d PendingKillPending: %d"), *Actor->GetName(), Actor->IsPendingKill(), Actor->IsPendingKillOrUnreachable(), Actor->IsPendingKillPending());
+	UE_LOG(LogTakesCore, Log, TEXT("Actor: %s IsValid: %d PendingKillOrUnreachable: %d PendingKillPending: %d"), *Actor->GetName(), ::IsValid(Actor), Actor->IsPendingKillOrUnreachable(), Actor->IsPendingKillPending());
 
 	UTakeRecorderSources* Sources = InSequence->FindOrAddMetaData<UTakeRecorderSources>();
 
diff --git a/Engine/Plugins/VirtualProduction/Takes/Source/TakeTrackRecorders/Private/MovieSceneAnimationTrackRecorder.cpp b/Engine/Plugins/VirtualProduction/Takes/Source/TakeTrackRecorders/Private/MovieSceneAnimationTrackRecorder.cpp
--- a/Engine/Plugins/VirtualProduction/Takes/Source/TakeTrackRecorders/Private/MovieSceneAnimationTrackRecorder.cpp
+++ b/Engine/Plugins/VirtualProduction/Takes/Source/TakeTrackRecorders/Private/MovieSceneAnimationTrackRecorder.cpp
@@ -262,8 +262,7 @@ void UMovieSceneAnimationTrackRecorder::RecordSampleImpl(const FQualifiedFrameTi
 			bRecordInWorldSpace = !OwningTakeRecorderSource->IsOtherActorBeingRecorded(AttachParent->GetOwner());
 		}
 
-		const TOptional<FQualifiedFrameTime> CurrentFrameTime = FApp::GetCurrentFrameTime();
-		FFrameRate SampleRate = CurrentFrameTime.IsSet() ? CurrentFrameTime.GetValue().Rate : MovieScene->GetDisplayRate();
+		FFrameRate SampleRate = MovieScene->GetDisplayRate();
 
 		//Set this up here so we know that it's parent sources have also been added so we record in the correct space
 		FAnimationRecordingSettings RecordingSettings;
diff --git a/Engine/Plugins/VirtualProduction/nDisplaySupportForLevelSnapshots/Source/nDisplaySupportForLevelSnapshots/Private/nDisplaySupportForLevelSnapshots.cpp b/Engine/Plugins/VirtualProduction/nDisplaySupportForLevelSnapshots/Source/nDisplaySupportForLevelSnapshots/Private/nDisplaySupportForLevelSnapshots.cpp
--- a/Engine/Plugins/VirtualProduction/nDisplaySupportForLevelSnapshots/Source/nDisplaySupportForLevelSnapshots/Private/nDisplaySupportForLevelSnapshots.cpp
+++ b/Engine/Plugins/VirtualProduction/nDisplaySupportForLevelSnapshots/Source/nDisplaySupportForLevelSnapshots/Private/nDisplaySupportForLevelSnapshots.cpp
@@ -15,6 +15,8 @@
 
 void FnDisplaySupportForLevelSnapshotsModule::StartupModule()
 {
+	UE_LOG(LogTemp, Log, TEXT("Starting Level Snapshot support for nDisplay..."));
+	
 	ILevelSnapshotsModule& LevelSnapshotsModule = static_cast<ILevelSnapshotsModule&>(FModuleManager::Get().LoadModuleChecked("LevelSnapshots"));
 
 	// Need to blacklist properties through which subobjects can be found
@@ -30,23 +32,7 @@ void FnDisplaySupportForLevelSnapshotsModule::StartupModule()
 }
 
 void FnDisplaySupportForLevelSnapshotsModule::ShutdownModule()
-{
-	if (!FModuleManager::Get().IsModuleLoaded("Level Snapshots"))
-	{
-		return;
-	}
-	ILevelSnapshotsModule& LevelSnapshotsModule = ILevelSnapshotsModule::Get();
-	
-	LevelSnapshotsModule.UnregisterCustomObjectSerializer(ADisplayClusterRootActor::StaticClass());
-	LevelSnapshotsModule.UnregisterCustomObjectSerializer(UDisplayClusterConfigurationData::StaticClass());
-	LevelSnapshotsModule.UnregisterCustomObjectSerializer(UDisplayClusterConfigurationCluster::StaticClass());
-	LevelSnapshotsModule.UnregisterCustomObjectSerializer(UDisplayClusterConfigurationClusterNode::StaticClass());
-
-	FDisplayClusterRootActorSerializer::UnblacklistCustomProperties(LevelSnapshotsModule);
-	FDisplayClusterConfigurationDataSerializer::UnblacklistCustomProperties(LevelSnapshotsModule);
-	FDisplayClusterConfigurationClusterSerializer::UnblacklistCustomProperties(LevelSnapshotsModule);
-	FDisplayClusterConfigurationClusterNodeSerializer::UnblacklistCustomProperties(LevelSnapshotsModule);
-}
+{}
 
 #undef LOCTEXT_NAMESPACE
 	
