Refresh UE5 branch with latest Level Snapshots which includes support for subobjects.

diff --git a/Engine/Plugins/VirtualProduction/LevelSnapshots/Source/LevelSnapshots/Private/ApplySnapshotFilter.cpp b/Engine/Plugins/VirtualProduction/LevelSnapshots/Source/LevelSnapshots/Private/ApplySnapshotFilter.cpp
--- a/Engine/Plugins/VirtualProduction/LevelSnapshots/Source/LevelSnapshots/Private/ApplySnapshotFilter.cpp
+++ b/Engine/Plugins/VirtualProduction/LevelSnapshots/Source/LevelSnapshots/Private/ApplySnapshotFilter.cpp
@@ -10,46 +10,20 @@
 #include "LevelSnapshotsLog.h"
 #include "LevelSnapshotsModule.h"
 #include "PropertyComparisonParams.h"
+#include "SnapshotCustomVersion.h"
 #include "Restorability/SnapshotRestorability.h"
 
 #include "Components/ActorComponent.h"
 #include "GameFramework/Actor.h"
 #include "Modules/ModuleManager.h"
+#include "Util/EquivalenceUtil.h"
 
 namespace
 {
-	UActorComponent* TryFindMatchingComponent(AActor* SnapshotActor, const TInlineComponentArray<UActorComponent*>& SnapshotComponents, AActor* WorldActor, UActorComponent* WorldComponent)
+	bool AreSubobjectsAllowed(UObject* SnapshotObject, UObject* WorldObject)
 	{
-		const FName ComponentName = WorldComponent->GetFName();
-		UActorComponent* const* PossibleResult = SnapshotComponents.FindByPredicate([ComponentName](UActorComponent* Component)
-		{
-			return Component->GetFName().IsEqual(ComponentName);
-		});
-		
-		if (!PossibleResult)
-		{
-			UE_LOG(LogLevelSnapshots, Warning, TEXT("Failed to match world component called '%s' of world actor called '%s' to any components in snapshot actor. Did you change the name of a component?"),
-                *WorldComponent->GetName(),
-                *WorldActor->GetName()
-            );
-			return nullptr;
-		}
-
-		UActorComponent* SnapshotComponent = *PossibleResult;
-		if (!ensure(SnapshotComponent->GetClass() == WorldComponent->GetClass()))
-		{
-			UE_LOG(LogLevelSnapshots, Error, TEXT("Snapshot component called  %s of snapshot actor called %s has class %s while world component called %s in world actor called %s had class %s. Components are expected to not change classes. Did you change the class of a component?"),
-                *SnapshotComponent->GetName(),
-                *SnapshotActor->GetName(),
-                *SnapshotComponent->GetClass()->GetName(),
-                *WorldActor->GetName(),
-                *WorldComponent->GetName(),
-                *WorldComponent->GetClass()->GetName()
-            );
-			return nullptr;
-		}
-
-		return SnapshotComponent;
+		return (WorldObject && FSnapshotRestorability::IsSubobjectDesirableForCapture(WorldObject))
+			|| (SnapshotObject && FSnapshotRestorability::IsSubobjectDesirableForCapture(SnapshotObject));
 	}
 }
 
@@ -70,12 +44,15 @@ void FApplySnapshotFilter::ApplyFilterToFindSelectedProperties(FPropertySelectio
 }
 
 FApplySnapshotFilter::FPropertyContainerContext::FPropertyContainerContext(
-	FPropertySelection& SelectionToAddTo, UStruct* ContainerClass, void* SnapshotContainer, void* WorldContainer, const TArray<FString>& AuthoredPathInformation, const FLevelSnapshotPropertyChain& PropertyChain, UClass* RootClass)
+	FPropertySelectionMap& MapToAddTo, FPropertySelection& SelectionToAddTo, UStruct* ContainerClass, void* SnapshotContainer, void* WorldContainer, UObject* AnalysedSnapshotObject, UObject* AnalysedWorldObject, const TArray<FString>& AuthoredPathInformation, const FLevelSnapshotPropertyChain& PropertyChain, UClass* RootClass)
 	:
+	MapToAddTo(MapToAddTo),
 	SelectionToAddTo(SelectionToAddTo),
 	ContainerClass(ContainerClass),
 	SnapshotContainer(SnapshotContainer),
 	WorldContainer(WorldContainer),
+	AnalysedSnapshotObject(AnalysedSnapshotObject),
+	AnalysedWorldObject(AnalysedWorldObject),
 	AuthoredPathInformation(AuthoredPathInformation),
 	PropertyChain(PropertyChain),
 	RootClass(RootClass)
@@ -118,33 +95,51 @@ bool FApplySnapshotFilter::EnsureParametersAreValid() const
 
 void FApplySnapshotFilter::AnalyseComponentProperties(FPropertySelectionMap& MapToAddTo)
 {
-	TInlineComponentArray<UActorComponent*> WorldComponents;
-	TInlineComponentArray<UActorComponent*> SnapshotComponents;
-	WorldActor->GetComponents(WorldComponents);
-	DeserializedSnapshotActor->GetComponents(SnapshotComponents);
-	for (UActorComponent* WorldComp : WorldComponents)
-	{
-		if (!FSnapshotRestorability::IsComponentDesirableForCapture(WorldComp))
+	FAddedAndRemovedComponentInfo ComponentSelection;
+
+	const bool bShouldAnalyseUnmatchedComponents = Snapshot->GetSerializedData().GetSnapshotVersionInfo().GetSnapshotCustomVersion() >= FSnapshotCustomVersion::SubobjectSupport;
+	SnapshotUtil::IterateComponents(DeserializedSnapshotActor, WorldActor,
+		[this, &MapToAddTo](UActorComponent* SnapshotComp, UActorComponent* WorldComp)
 		{
-			UE_LOG(LogLevelSnapshots, Verbose, TEXT("Skipping world component '%s' of world actor '%s'"), *WorldComp->GetName(), *WorldActor->GetName());
-			continue;
-		}
-		
-		if (UActorComponent* SnapshotMatchedComp = TryFindMatchingComponent(DeserializedSnapshotActor, SnapshotComponents, WorldActor, WorldComp))
+			FilterSubobjectPair(MapToAddTo, SnapshotComp, WorldComp);
+		},
+		[&ComponentSelection, bShouldAnalyseUnmatchedComponents](UActorComponent* UnmatchedSnapshotComp)
+		{
+			if (bShouldAnalyseUnmatchedComponents)
+			{
+				ComponentSelection.SnapshotComponentsToAdd.Add(UnmatchedSnapshotComp);
+			}
+		},
+		[&ComponentSelection, bShouldAnalyseUnmatchedComponents](UActorComponent* UnmatchedWorldComp)
 		{
-			FilterSubobjectPair(MapToAddTo, SnapshotMatchedComp, WorldComp);
+			
+			if (bShouldAnalyseUnmatchedComponents)
+			{
+				ComponentSelection.EditorWorldComponentsToRemove.Add(UnmatchedWorldComp);
+			}
 		}
+	);
+	
+	const bool bIsEmpty = ComponentSelection.SnapshotComponentsToAdd.Num() == 0 && ComponentSelection.EditorWorldComponentsToRemove.Num() == 0;
+	if (!bIsEmpty)
+	{
+		MapToAddTo.AddComponentSelection(WorldActor, ComponentSelection);
 	}
 }
 
 void FApplySnapshotFilter::FilterActorPair(FPropertySelectionMap& MapToAddTo)
 {
+	AnalysedSnapshotObjects.Add(DeserializedSnapshotActor);
+	
 	FPropertySelection ActorSelection;
 	FPropertyContainerContext ActorContext(
+		MapToAddTo,
 		ActorSelection,
         DeserializedSnapshotActor->GetClass(),
         DeserializedSnapshotActor,
         WorldActor,
+        DeserializedSnapshotActor,
+        WorldActor,
         {},
         FLevelSnapshotPropertyChain(),
         WorldActor->GetClass()
@@ -154,46 +149,79 @@ void FApplySnapshotFilter::FilterActorPair(FPropertySelectionMap& MapToAddTo)
 	const EFilterObjectPropertiesResult FilterResult = FindAndFilterCustomSubobjectPairs(MapToAddTo, DeserializedSnapshotActor, WorldActor);
 	
 	ActorSelection.SetHasCustomSerializedSubobjects(FilterResult == EFilterObjectPropertiesResult::HasCustomSubobjects);
-	MapToAddTo.AddObjectProperties(WorldActor, ActorSelection);
+	if (!ActorSelection.IsEmpty())
+	{
+		MapToAddTo.AddObjectProperties(WorldActor, ActorSelection);
+	}
 }
 
-bool FApplySnapshotFilter::FilterSubobjectPair(FPropertySelectionMap& MapToAddTo, UObject* SnapshotSubobject, UObject* WorldSubobject)
+FApplySnapshotFilter::EPropertySearchResult FApplySnapshotFilter::FilterSubobjectPair(FPropertySelectionMap& MapToAddTo, UObject* SnapshotSubobject, UObject* WorldSubobject)
 {
+	if (const EPropertySearchResult* Result = AnalysedSnapshotObjects.Find(SnapshotSubobject))
+	{
+		return *Result;
+	}
+	if (SnapshotSubobject->GetClass() != WorldSubobject->GetClass())
+	{
+		return EPropertySearchResult::FoundProperties;
+	}
+	
 	FPropertySelection SubobjectSelection;
 	FPropertyContainerContext ComponentContext(
+		MapToAddTo,
 		SubobjectSelection,
         SnapshotSubobject->GetClass(),
         SnapshotSubobject,
         WorldSubobject,
+        SnapshotSubobject,
+        WorldSubobject,
         { WorldSubobject->GetName() },
         FLevelSnapshotPropertyChain(),
 		WorldSubobject->GetClass()
         );
 	
 	AnalyseRootProperties(ComponentContext, SnapshotSubobject, WorldSubobject);
+	
 	const EFilterObjectPropertiesResult FilterResult = FindAndFilterCustomSubobjectPairs(MapToAddTo, SnapshotSubobject, WorldSubobject);
+	SubobjectSelection.SetHasCustomSerializedSubobjects(FilterResult == EFilterObjectPropertiesResult::HasCustomSubobjects);
 
-	if (ensureMsgf(MapToAddTo.GetSelectedProperties(WorldSubobject) == nullptr, TEXT("Object %s was analysed more than once. Most likely an ICustomObjectSnapshotSerializer implementation returned the same object pair more than once or returned an object pair that was already discovered by the standard system, e.g. UPROPERTY(EditAnywhere, Instanced) UObject* Object."), *WorldSubobject->GetName()))
+	// Non-component subobjects should always be added so when applying to world and resolving a subobject, we can tell whether it is referenced or just a dead object still existing in memory
+	const bool bAddedProperties = [&MapToAddTo, WorldSubobject, &SubobjectSelection]()
 	{
-		SubobjectSelection.SetHasCustomSerializedSubobjects(FilterResult == EFilterObjectPropertiesResult::HasCustomSubobjects);
-		return MapToAddTo.AddObjectProperties(WorldSubobject, SubobjectSelection);
-	}
-
-	return false;
+		if (SubobjectSelection.IsEmpty())
+		{
+			// We don't track components because we can efficiently tell whether they're dead or not by asking the parent actor
+			if (!WorldSubobject->IsA<UActorComponent>())
+			{
+				MapToAddTo.MarkSubobjectForRestoringReferencesButSkipProperties(WorldSubobject);
+			}
+			return false;
+		}
+		
+		MapToAddTo.AddObjectProperties(WorldSubobject, SubobjectSelection);
+		return true;
+	}();
+	EPropertySearchResult Result = bAddedProperties ? EPropertySearchResult::FoundProperties : EPropertySearchResult::NoPropertiesFound;
+	AnalysedSnapshotObjects.Add(SnapshotSubobject, Result);
+	return Result;
 }
 
-void FApplySnapshotFilter::FilterStructPair(FPropertyContainerContext& Parent, FStructProperty* StructProperty)
+FApplySnapshotFilter::EPropertySearchResult FApplySnapshotFilter::FilterStructPair(FPropertyContainerContext& Parent, FStructProperty* StructProperty)
 {
-	FPropertyContainerContext StructContext(Parent.SelectionToAddTo,
+	FPropertyContainerContext StructContext(
+		Parent.MapToAddTo,
+		Parent.SelectionToAddTo,
         StructProperty->Struct,
         StructProperty->ContainerPtrToValuePtr<uint8>(Parent.SnapshotContainer),
         StructProperty->ContainerPtrToValuePtr<uint8>(Parent.WorldContainer),
+        Parent.AnalysedSnapshotObject,
+        Parent.AnalysedWorldObject,
         Parent.AuthoredPathInformation,
         Parent.PropertyChain.MakeAppended(StructProperty),
         Parent.RootClass
         );
 	StructContext.AuthoredPathInformation.Add(StructProperty->GetAuthoredName());
-	AnalyseStructProperties(StructContext);
+	return AnalyseStructProperties(StructContext);
 }
 
 FApplySnapshotFilter::EFilterObjectPropertiesResult FApplySnapshotFilter::FindAndFilterCustomSubobjectPairs(FPropertySelectionMap& MapToAddTo, UObject* SnapshotOwner, UObject* WorldOwner)
@@ -223,16 +251,30 @@ FApplySnapshotFilter::EFilterObjectPropertiesResult FApplySnapshotFilter::FindAn
 		UObject* SnapshotSubobject = ExternalSerializer->FindOrRecreateSubobjectInSnapshotWorld(SnapshotOwner, *SubobjectMetadata, SubobjectDataReader);
 		UObject* EditorSubobject = ExternalSerializer->FindSubobjectInEditorWorld(WorldOwner, *SubobjectMetadata, SubobjectDataReader);
 		
-		// External modules implement this: must be validated
-		if (!SnapshotSubobject || !EditorSubobject)
+		// Snapshot subobject is missing
+		if (!SnapshotSubobject || !ensureAlwaysMsgf(SnapshotSubobject->IsIn(SnapshotOwner), TEXT("Your interface must return subobjects")))
 		{
-			UE_CLOG(SnapshotSubobject && !EditorSubobject, LogLevelSnapshots, Warning, TEXT("Restoring missing subobjects is not supported"));
+			UE_CLOG(!SnapshotSubobject && EditorSubobject, LogLevelSnapshots, Warning, TEXT("Snapshot subobject was missing. Removing newly added editor world subobjects is not supported."));
 			continue;
 		}
 
-		if (ensureAlwaysMsgf(SnapshotSubobject->IsIn(SnapshotOwner) && EditorSubobject->IsIn(WorldOwner), TEXT("Your interface must return subobjects")))
+		// Editor subobject is missing
+		if (!EditorSubobject)
 		{
-			bAtLeastOneSubobjectWasAdded |= FilterSubobjectPair(MapToAddTo, SnapshotSubobject, EditorSubobject);
+			if (ensureAlwaysMsgf(SnapshotSubobject->IsIn(SnapshotOwner), TEXT("Your interface must return subobjects")))
+			{
+				MapToAddTo.AddCustomEditorSubobjectToRecreate(WorldOwner, SnapshotSubobject);
+				bAtLeastOneSubobjectWasAdded = true;
+				AnalysedSnapshotObjects.Add(SnapshotSubobject);
+			}
+			continue;
+		}
+
+		// Both subobjects exist
+		if (ensureAlwaysMsgf(EditorSubobject->IsIn(WorldOwner), TEXT("Your interface must return subobjects"))
+			&& ensureAlwaysMsgf(!AnalysedSnapshotObjects.Contains(SnapshotSubobject), TEXT("You returned an object that was already found (probably by standard Level Snapshot serialisation)")))
+		{
+			bAtLeastOneSubobjectWasAdded |= FilterSubobjectPair(MapToAddTo, SnapshotSubobject, EditorSubobject) == EPropertySearchResult::FoundProperties;
 		}
 	}
 
@@ -263,44 +305,301 @@ void FApplySnapshotFilter::AnalyseRootProperties(FPropertyContainerContext& Cont
 			break;
 		}
 		
-		const ECheckSubproperties CheckSubpropertyBehaviour = AnalyseProperty(ContainerContext, *FieldIt, bSkipEqualityTest);
-		if (!bSkipEqualityTest && CheckSubpropertyBehaviour == ECheckSubproperties::CheckSubproperties)
+		AnalyseProperty(ContainerContext, *FieldIt, bSkipEqualityTest);
+	}
+}
+
+TOptional<FApplySnapshotFilter::EPropertySearchResult> FApplySnapshotFilter::HandlePossibleStructProperties(FPropertyContainerContext& ContainerContext, FProperty* PropertyToHandle)
+{
+	if (FStructProperty* StructProperty = CastField<FStructProperty>(PropertyToHandle))
+	{
+		return FilterStructPair(ContainerContext, StructProperty);
+	}
+	return {};
+}
+
+FApplySnapshotFilter::EPropertySearchResult FApplySnapshotFilter::AnalyseStructProperties(FPropertyContainerContext& ContainerContext)
+{
+	bool bFoundAtLeastOneProperty = false;
+	for (TFieldIterator<FProperty> FieldIt(ContainerContext.ContainerClass); FieldIt; ++FieldIt)
+	{
+		bFoundAtLeastOneProperty |= AnalyseProperty(ContainerContext, *FieldIt) == EPropertySearchResult::FoundProperties;
+	}
+	return bFoundAtLeastOneProperty ? EPropertySearchResult::FoundProperties : EPropertySearchResult::NoPropertiesFound;
+}
+
+TOptional<FApplySnapshotFilter::EPropertySearchResult> FApplySnapshotFilter::HandlePossibleSubobjectProperties(FPropertyContainerContext& ContainerContext, FProperty* PropertyToHandle)
+{
+	if (FArrayProperty* ArrayProperty = CastField<FArrayProperty>(PropertyToHandle))
+	{
+		return AnalysePossibleArraySubobjectProperties(ContainerContext, ArrayProperty);
+	}
+	if (FSetProperty* SetProperty = CastField<FSetProperty>(PropertyToHandle))
+	{
+		return AnalysePossibleSetSubobjectProperties(ContainerContext, SetProperty);
+	}
+	if (FMapProperty* MapProperty = CastField<FMapProperty>(PropertyToHandle))
+	{
+		return AnalysePossibleMapSubobjectProperties(ContainerContext, MapProperty);
+	}
+	
+	if (FObjectPropertyBase* ObjectProperty = CastField<FObjectPropertyBase>(PropertyToHandle))
+	{
+		void* SnapshotValuePtr = PropertyToHandle->ContainerPtrToValuePtr<void>(ContainerContext.SnapshotContainer, 0);
+		void* WorldValuePtr = PropertyToHandle->ContainerPtrToValuePtr<void>(ContainerContext.WorldContainer, 0);
+		return AnalysePossibleSubobjectProperties(ContainerContext, ObjectProperty, SnapshotValuePtr, WorldValuePtr);
+	}
+	
+	return {};
+}
+
+TOptional<FApplySnapshotFilter::EPropertySearchResult> FApplySnapshotFilter::AnalysePossibleArraySubobjectProperties(FPropertyContainerContext& ContainerContext, FArrayProperty* PropertyToHandle)
+{
+	struct FArrayData
+	{
+		FScriptArrayHelper SnapshotArray;
+		FScriptArrayHelper WorldArray;
+		const int32 SnapshotNum;
+		const int32 WorldNum;
+
+		FArrayData(FPropertyContainerContext& ContainerContext, FArrayProperty* ArrayProperty)
+			:
+			SnapshotArray(FScriptArrayHelper(ArrayProperty, ArrayProperty->ContainerPtrToValuePtr<void>(ContainerContext.SnapshotContainer, 0))),
+			WorldArray(FScriptArrayHelper(ArrayProperty, ArrayProperty->ContainerPtrToValuePtr<void>(ContainerContext.WorldContainer, 0))),
+			SnapshotNum(SnapshotArray.Num()),
+			WorldNum(WorldArray.Num())
+		{}
+
+		void* GetSnapshotValuePtrFromIndex(int32 WorldIndex)
+		{
+			return SnapshotArray.GetRawPtr(WorldIndex);
+		}
+
+		void* GetWorldValuePtrFromIndex(int32 WorldIndex)
 		{
-			HandleStructProperties(ContainerContext, *FieldIt);
-			// TODO: To analyse subobjects, you'd add another utility function here
+			return WorldArray.GetRawPtr(WorldIndex);
 		}
+	};
+
+	FObjectPropertyBase* InnerAsObjectProperty = CastField<FObjectPropertyBase>(PropertyToHandle->Inner);
+	if (InnerAsObjectProperty && !CastField<FSoftObjectProperty>(PropertyToHandle->Inner))
+	{
+		FArrayData ArrayData(ContainerContext, PropertyToHandle);
+		return AnalysePossibleSubobjectsInCollection(ContainerContext, InnerAsObjectProperty, ArrayData);
 	}
+	return {};
 }
 
-void FApplySnapshotFilter::AnalyseStructProperties(FPropertyContainerContext& ContainerContext)
+TOptional<FApplySnapshotFilter::EPropertySearchResult> FApplySnapshotFilter::AnalysePossibleSetSubobjectProperties(FPropertyContainerContext& ContainerContext, FSetProperty* PropertyToHandle)
 {
-	for (TFieldIterator<FProperty> FieldIt(ContainerContext.ContainerClass); FieldIt; ++FieldIt)
+	struct FSetData
 	{
-		const ECheckSubproperties CheckSubpropertyBehaviour = AnalyseProperty(ContainerContext, *FieldIt);
-		if (CheckSubpropertyBehaviour == ECheckSubproperties::CheckSubproperties)
+		FScriptSetHelper SnapshotSet;
+		FScriptSetHelper WorldSet;
+		const int32 SnapshotNum;
+		const int32 WorldNum;
+
+		TMap<FName, void*> ObjectNameToSnapshotValuePtr;
+		const FObjectPropertyBase* ObjectProperty;  
+
+		FSetData(FPropertyContainerContext& ContainerContext, FSetProperty* SetProperty)
+			:
+			SnapshotSet(FScriptSetHelper(SetProperty, SetProperty->ContainerPtrToValuePtr<void>(ContainerContext.SnapshotContainer, 0))),
+			WorldSet(FScriptSetHelper(SetProperty, SetProperty->ContainerPtrToValuePtr<void>(ContainerContext.WorldContainer, 0))),
+			SnapshotNum(SnapshotSet.Num()),
+			WorldNum(WorldSet.Num()),
+			ObjectProperty(CastField<FObjectPropertyBase>(SetProperty->ElementProp))
+		{
+			for (int32 SnapshotIndex = 0; SnapshotIndex < SnapshotSet.Num(); ++SnapshotIndex)
+			{
+				void* SnapshotValuePtr = SnapshotSet.GetElementPtr(SnapshotIndex);
+				if (const UObject* SnapshotObject = ObjectProperty->GetObjectPropertyValue(SnapshotValuePtr))
+				{
+					ObjectNameToSnapshotValuePtr.Add(SnapshotObject->GetFName(), SnapshotValuePtr);
+				}
+			}
+		}
+
+		void* GetSnapshotValuePtrFromIndex(int32 WorldIndex)
+		{
+			// Note: we only look for subobjects with the same name because it is faster; otherwise we'd need to compare all properties (n^2 performance).
+			// This works most of the time but not when some subobject has changed name (unlikely but possible).
+			// In the cases it does not work, we just end up with an additional diff in the results view; a low price to pay.
+			const void* WorldValuePtr = GetWorldValuePtrFromIndex(WorldIndex);
+			if (UObject* WorldObject = ObjectProperty->GetObjectPropertyValue(WorldValuePtr))
+			{
+				void** PossibleSnapshotObject = ObjectNameToSnapshotValuePtr.Find(WorldObject->GetFName());
+				return PossibleSnapshotObject ? *PossibleSnapshotObject : nullptr;
+			}
+			return nullptr;
+		}
+
+		void* GetWorldValuePtrFromIndex(int32 WorldIndex)
 		{
-			HandleStructProperties(ContainerContext, *FieldIt);
-			// TODO: To analyse subobjects, you'd add another utility function here
+			return WorldSet.GetElementPtr(WorldIndex);
 		}
+	};
+
+	FObjectPropertyBase* InnerAsObjectProperty = CastField<FObjectPropertyBase>(PropertyToHandle->ElementProp);
+	if (InnerAsObjectProperty && !CastField<FSoftObjectProperty>(PropertyToHandle->ElementProp))
+	{
+		FSetData SetData(ContainerContext, PropertyToHandle);
+		return AnalysePossibleSubobjectsInCollection(ContainerContext, InnerAsObjectProperty, SetData);
 	}
+	return {};
 }
 
-void FApplySnapshotFilter::HandleStructProperties(FPropertyContainerContext& ContainerContext, FProperty* PropertyToHandle)
+TOptional<FApplySnapshotFilter::EPropertySearchResult> FApplySnapshotFilter::AnalysePossibleMapSubobjectProperties(FPropertyContainerContext& ContainerContext, FMapProperty* PropertyToHandle)
 {
-	if (FStructProperty* StructProperty = CastField<FStructProperty>(PropertyToHandle))
+	struct FMapData
+	{
+		FScriptMapHelper SnapshotMap;
+		FScriptMapHelper WorldMap;
+		const int32 SnapshotNum;
+		const int32 WorldNum;
+
+		FMapData(FPropertyContainerContext& ContainerContext, FMapProperty* MapProperty)
+			:
+			SnapshotMap(FScriptMapHelper(MapProperty, MapProperty->ContainerPtrToValuePtr<void>(ContainerContext.SnapshotContainer, 0))),
+			WorldMap(FScriptMapHelper(MapProperty, MapProperty->ContainerPtrToValuePtr<void>(ContainerContext.WorldContainer, 0))),
+			SnapshotNum(SnapshotMap.Num()),
+			WorldNum(WorldMap.Num())
+		{}
+
+		void* GetSnapshotValuePtrFromIndex(int32 WorldIndex)
+		{
+			const void* WorldKey = WorldMap.GetKeyPtr(WorldIndex);
+			return SnapshotMap.FindValueFromHash(WorldKey);
+		}
+
+		void* GetWorldValuePtrFromIndex(int32 WorldIndex)
+		{
+			return WorldMap.GetValuePtr(WorldIndex);
+		}
+	};
+
+	FObjectPropertyBase* InnerAsObjectProperty = CastField<FObjectPropertyBase>(PropertyToHandle->ValueProp);
+	if (InnerAsObjectProperty && !CastField<FSoftObjectProperty>(PropertyToHandle->ValueProp))
+	{
+		FMapData MapData(ContainerContext, PropertyToHandle);
+		return AnalysePossibleSubobjectsInCollection(ContainerContext, InnerAsObjectProperty, MapData);
+	}
+	return {};
+}
+
+template<typename TCollectionData>
+TOptional<FApplySnapshotFilter::EPropertySearchResult> FApplySnapshotFilter::AnalysePossibleSubobjectsInCollection(FPropertyContainerContext& ContainerContext, FObjectPropertyBase* ObjectProperty, TCollectionData& Detail)
+{
+	TArray<UObject*> UnchangedSubobjects;
+	bool bFoundChangedProperties = false;
+	bool bHadUnmatchedObjects = Detail.SnapshotNum != Detail.WorldNum;
+	const int32 MaxSharedSize = FMath::Min(Detail.SnapshotNum, Detail.WorldNum);
+	for (int32 WorldIndex = 0; WorldIndex < MaxSharedSize; ++WorldIndex)
+	{
+		void* SnapshotValuePtr = Detail.GetSnapshotValuePtrFromIndex(WorldIndex);
+		if (!SnapshotValuePtr)
+		{
+			bHadUnmatchedObjects = true;
+			continue;
+		}
+
+		void* WorldValuePtr = Detail.GetWorldValuePtrFromIndex(WorldIndex);
+		const TOptional<EPropertySearchResult> SubbjectPairResult = AnalysePossibleSubobjectProperties(ContainerContext, ObjectProperty, SnapshotValuePtr, WorldValuePtr);
+		if (!SubbjectPairResult.IsSet())
+		{
+			return {};
+		}
+		
+		const bool bFoundChangesOnSubobject = SubbjectPairResult == EPropertySearchResult::FoundProperties;
+		bFoundChangedProperties |= bFoundChangesOnSubobject;
+		UObject* WorldObject = ObjectProperty->GetObjectPropertyValue(WorldValuePtr);
+		if (!bFoundChangesOnSubobject && WorldObject)
+		{
+			UnchangedSubobjects.Add(WorldObject);
+		}
+	}
+
+	const bool bHasCollectionChanged = bFoundChangedProperties || bHadUnmatchedObjects;
+	if (bHasCollectionChanged)
+	{
+		// The collection has changed. When restoring, SnapshotUtil::ResolveObjectDependencyForEditorWorld is called for each object
+		// We tell it to not write any data into the subobjects by adding an empty set.
+		for (UObject* UnchangedSubobject : UnchangedSubobjects)
+		{
+			ContainerContext.MapToAddTo.MarkSubobjectForRestoringReferencesButSkipProperties(UnchangedSubobject);
+		}
+	}
+	
+	return bHasCollectionChanged ? EPropertySearchResult::FoundProperties : EPropertySearchResult::NoPropertiesFound;
+}
+
+TOptional<FApplySnapshotFilter::EPropertySearchResult> FApplySnapshotFilter::AnalysePossibleSubobjectProperties(FPropertyContainerContext& ContainerContext, FObjectPropertyBase* PropertyToHandle, void* SnapshotValuePtr, void* WorldValuePtr)
+{
+	// Soft object properties do not work with subobjects and are not expected to: the path will remain pointing at the real world object...
+	if (FSoftObjectProperty* SoftObjectPath = CastField<FSoftObjectProperty>(PropertyToHandle))
+	{
+		// ... with soft paths we care about the string change only so compare it normally
+		return {};
+	}
+	
+	const bool bPropertyMarkedAsSubobject = PropertyToHandle->HasAnyPropertyFlags(CPF_PersistentInstance | CPF_InstancedReference | CPF_ContainsInstancedReference);
+	UObject* SnapshotObject = PropertyToHandle->GetObjectPropertyValue(SnapshotValuePtr);
+	UObject* WorldObject = PropertyToHandle->GetObjectPropertyValue(WorldValuePtr);
+
+	if (SnapshotObject == nullptr || WorldObject == nullptr)
+	{
+		// Are they subobjects
+		const bool bSnapshotObjectIsSubobject = bPropertyMarkedAsSubobject || (SnapshotObject && SnapshotObject->IsIn(DeserializedSnapshotActor));
+		const bool bWorldObjectIsSubobject = bPropertyMarkedAsSubobject || (WorldObject && WorldObject->IsIn(WorldActor));
+
+		// Blacklisted
+		const bool bIsSnapshotSupportedSubobject = SnapshotObject && bSnapshotObjectIsSubobject && !FSnapshotRestorability::IsSubobjectDesirableForCapture(SnapshotObject);
+		const bool bIsWorldSupportedSubobject = WorldObject && bWorldObjectIsSubobject && !FSnapshotRestorability::IsSubobjectDesirableForCapture(WorldObject);
+
+		// Migration: We did not save subobject data. In this case snapshot version resolves to nullptr.
+		const bool bIsOldSnapshot = Snapshot->GetSerializedData().GetSnapshotVersionInfo().GetSnapshotCustomVersion() < FSnapshotCustomVersion::SubobjectSupport;
+		const bool bOldSnapshotDataDidNotCaptureSubobjects = bIsOldSnapshot && bWorldObjectIsSubobject && SnapshotObject == nullptr;
+		
+		return SnapshotObject == WorldObject || bOldSnapshotDataDidNotCaptureSubobjects|| bIsSnapshotSupportedSubobject || bIsWorldSupportedSubobject
+			? 
+			EPropertySearchResult::NoPropertiesFound 
+			: 
+			bWorldObjectIsSubobject|| bSnapshotObjectIsSubobject ? EPropertySearchResult::FoundProperties : TOptional<FApplySnapshotFilter::EPropertySearchResult>{};
+	}
+
+	// Components are special: Only check whether they point to the same objects. 
+	const bool bAreComponents = SnapshotObject->IsA<UActorComponent>() || WorldObject->IsA<UActorComponent>();
+	if (bAreComponents)
+	{
+		return SnapshotUtil::AreObjectPropertiesEquivalent(Snapshot->GetSerializedData(), PropertyToHandle, SnapshotValuePtr, WorldValuePtr, DeserializedSnapshotActor, WorldActor)
+			? EPropertySearchResult::NoPropertiesFound : EPropertySearchResult::FoundProperties;
+	}
+		
+	const bool bIsSubobject = bPropertyMarkedAsSubobject || SnapshotObject->IsIn(DeserializedSnapshotActor);
+	if (bIsSubobject)
 	{
-		FilterStructPair(ContainerContext, StructProperty);
+		return AreSubobjectsAllowed(SnapshotObject, WorldObject) 
+				? FilterSubobjectPair(ContainerContext.MapToAddTo, SnapshotObject, WorldObject) : EPropertySearchResult::NoPropertiesFound;
 	}
+
+	return {};
 }
 
-FApplySnapshotFilter::ECheckSubproperties FApplySnapshotFilter::AnalyseProperty(FPropertyContainerContext& ContainerContext, FProperty* PropertyInCommon, bool bSkipEqualityTest)
+FApplySnapshotFilter::EPropertySearchResult FApplySnapshotFilter::AnalyseProperty(FPropertyContainerContext& ContainerContext, FProperty* PropertyInCommon, bool bSkipEqualityTest)
 {
+	if (!FSnapshotRestorability::IsRestorableProperty(PropertyInCommon))
+	{
+		return EPropertySearchResult::NoPropertiesFound;
+	}
+
 	check(ContainerContext.WorldContainer);
 	check(ContainerContext.SnapshotContainer); 
 
-	if (!bSkipEqualityTest && !bAllowUnchangedProperties && Snapshot->AreSnapshotAndOriginalPropertiesEquivalent(PropertyInCommon, ContainerContext.SnapshotContainer, ContainerContext.WorldContainer, DeserializedSnapshotActor, WorldActor))
+	if (!bSkipEqualityTest
+		&& !bAllowUnchangedProperties
+		&& ArePropertyValuesIdentical(ContainerContext, PropertyInCommon))
 	{
-		return SkipSubproperties;
+		return EPropertySearchResult::NoPropertiesFound;
 	}
 
 	// Now we can ask user whether they care about the property
@@ -313,7 +612,24 @@ FApplySnapshotFilter::ECheckSubproperties FApplySnapshotFilter::AnalyseProperty(
 	if (bIsPropertyValid)
 	{
 		ContainerContext.SelectionToAddTo.AddProperty(ContainerContext.PropertyChain.MakeAppended(PropertyInCommon));
-		return CheckSubproperties;
+		return EPropertySearchResult::FoundProperties;
 	}
-	return SkipSubproperties;
+	return EPropertySearchResult::NoPropertiesFound;
+}
+
+bool FApplySnapshotFilter::ArePropertyValuesIdentical(FPropertyContainerContext& ContainerContext, FProperty* PropertyInCommon)
+{
+	const TOptional<EPropertySearchResult> StructResult = HandlePossibleStructProperties(ContainerContext, PropertyInCommon);
+	if (StructResult.IsSet())
+	{
+		return *StructResult == EPropertySearchResult::NoPropertiesFound;
+	}
+	
+	const TOptional<EPropertySearchResult> SubobjectResult = HandlePossibleSubobjectProperties(ContainerContext, PropertyInCommon);
+	if (SubobjectResult.IsSet())
+	{
+		return *SubobjectResult == EPropertySearchResult::NoPropertiesFound;
+	}
+
+	return Snapshot->AreSnapshotAndOriginalPropertiesEquivalent(PropertyInCommon, ContainerContext.SnapshotContainer, ContainerContext.WorldContainer, DeserializedSnapshotActor, WorldActor);
 }
diff --git a/Engine/Plugins/VirtualProduction/LevelSnapshots/Source/LevelSnapshots/Private/ApplySnapshotFilter.h b/Engine/Plugins/VirtualProduction/LevelSnapshots/Source/LevelSnapshots/Private/ApplySnapshotFilter.h
--- a/Engine/Plugins/VirtualProduction/LevelSnapshots/Source/LevelSnapshots/Private/ApplySnapshotFilter.h
+++ b/Engine/Plugins/VirtualProduction/LevelSnapshots/Source/LevelSnapshots/Private/ApplySnapshotFilter.h
@@ -39,11 +39,14 @@ private:
 	
 	struct FPropertyContainerContext
 	{
+		FPropertySelectionMap& MapToAddTo;
 		FPropertySelection& SelectionToAddTo;
 	
 		UStruct* ContainerClass;
 		void* SnapshotContainer;
 		void* WorldContainer;
+		UObject* AnalysedSnapshotObject;
+		UObject* AnalysedWorldObject;
 			
 		/* Information passed to blueprints. Property name is appended to this.
 		 * Example: [FooComponent] [BarStructPropertyName]...
@@ -55,7 +58,13 @@ private:
 		/* Class that PropertyChain begins from. */
 		UClass* RootClass;
 
-		FPropertyContainerContext(FPropertySelection& SelectionToAddTo, UStruct* ContainerClass, void* SnapshotContainer, void* WorldContainer, const TArray<FString>& AuthoredPathInformation, const FLevelSnapshotPropertyChain& PropertyChain, UClass* RootClass);
+		FPropertyContainerContext(FPropertySelectionMap& MapToAddTo, FPropertySelection& SelectionToAddTo, UStruct* ContainerClass, void* SnapshotContainer, void* WorldContainer, UObject* AnalysedSnapshotObject, UObject* AnalysedWorldObject, const TArray<FString>& AuthoredPathInformation, const FLevelSnapshotPropertyChain& PropertyChain, UClass* RootClass);
+	};
+
+	enum class EPropertySearchResult
+	{
+		FoundProperties,
+		NoPropertiesFound
 	};
 	
 	FApplySnapshotFilter(ULevelSnapshot* Snapshot, AActor* DeserializedSnapshotActor, AActor* WorldActor, const ULevelSnapshotFilter* Filter);
@@ -64,32 +73,42 @@ private:
 	void AnalyseComponentProperties(FPropertySelectionMap& MapToAddTo);
 
 	void FilterActorPair(FPropertySelectionMap& MapToAddTo);
-	/* @return Whether any property selection was added */
-	bool FilterSubobjectPair(FPropertySelectionMap& MapToAddTo, UObject* SnapshotSubobject, UObject* WorldSubobject);
-	void FilterStructPair(FPropertyContainerContext& Parent, FStructProperty* StructProperty);
+	EPropertySearchResult FilterSubobjectPair(FPropertySelectionMap& MapToAddTo, UObject* SnapshotSubobject, UObject* WorldSubobject);
+	EPropertySearchResult FilterStructPair(FPropertyContainerContext& Parent, FStructProperty* StructProperty);
 	EFilterObjectPropertiesResult FindAndFilterCustomSubobjectPairs(FPropertySelectionMap& MapToAddTo, UObject* SnapshotOwner, UObject* WorldOwner);
 
 	void AnalyseRootProperties(FPropertyContainerContext& ContainerContext, UObject* SnapshotObject, UObject* WorldObject);
-	void AnalyseStructProperties(FPropertyContainerContext& ContainerContext);
-	void HandleStructProperties(FPropertyContainerContext& ContainerContext, FProperty* PropertyToHandle);
+	TOptional<EPropertySearchResult> HandlePossibleStructProperties(FPropertyContainerContext& ContainerContext, FProperty* PropertyToHandle);
+	EPropertySearchResult AnalyseStructProperties(FPropertyContainerContext& ContainerContext);
+
 	
-	enum ECheckSubproperties
-	{
-		CheckSubproperties,
-        SkipSubproperties
-    }; 
-	ECheckSubproperties AnalyseProperty(FPropertyContainerContext& ContainerContext, FProperty* PropertyInCommon, bool bSkipEqualityTest = false);
+	TOptional<EPropertySearchResult> HandlePossibleSubobjectProperties(FPropertyContainerContext& ContainerContext, FProperty* PropertyToHandle);
+	
+	TOptional<EPropertySearchResult> AnalysePossibleArraySubobjectProperties(FPropertyContainerContext& ContainerContext, FArrayProperty* PropertyToHandle);
+	TOptional<EPropertySearchResult> AnalysePossibleSetSubobjectProperties(FPropertyContainerContext& ContainerContext, FSetProperty* PropertyToHandle);
+	TOptional<EPropertySearchResult> AnalysePossibleMapSubobjectProperties(FPropertyContainerContext& ContainerContext, FMapProperty* PropertyToHandle);
+	
+	template<typename TCollectionData>
+	TOptional<EPropertySearchResult> AnalysePossibleSubobjectsInCollection(FPropertyContainerContext& ContainerContext, FObjectPropertyBase* ObjectProperty, TCollectionData& Detail);
+	
+	TOptional<EPropertySearchResult> AnalysePossibleSubobjectProperties(FPropertyContainerContext& ContainerContext, FObjectPropertyBase* PropertyToHandle, void* SnapshotValuePtr, void* WorldValuePtr);
+
+	
+
+	EPropertySearchResult AnalyseProperty(FPropertyContainerContext& ContainerContext, FProperty* PropertyInCommon, bool bSkipEqualityTest = false);
+	bool ArePropertyValuesIdentical(FPropertyContainerContext& ContainerContext, FProperty* PropertyInCommon);
 
 	
 	ULevelSnapshot* Snapshot;
 	AActor* DeserializedSnapshotActor;
 	AActor* WorldActor;
 	const ULevelSnapshotFilter* Filter;
+	/** Every time an object is analysed, it is added here, so we do not analyse more than once. */
+	TMap<UObject*, EPropertySearchResult> AnalysedSnapshotObjects;
 
 	/* Do we allow properties that do not show up in the details panel? */
 	bool bAllowNonEditableProperties = false;
 
 	/* Do we allow adding properties that did not change to the selection map? */
 	bool bAllowUnchangedProperties = false;
-	
 };
diff --git a/Engine/Plugins/VirtualProduction/LevelSnapshots/Source/LevelSnapshots/Private/Archive/ApplySnapshotDataArchiveV2.cpp b/Engine/Plugins/VirtualProduction/LevelSnapshots/Source/LevelSnapshots/Private/Archive/ApplySnapshotDataArchiveV2.cpp
--- a/Engine/Plugins/VirtualProduction/LevelSnapshots/Source/LevelSnapshots/Private/Archive/ApplySnapshotDataArchiveV2.cpp
+++ b/Engine/Plugins/VirtualProduction/LevelSnapshots/Source/LevelSnapshots/Private/Archive/ApplySnapshotDataArchiveV2.cpp
@@ -15,7 +15,9 @@
 #include "Serialization/ObjectWriter.h"
 #include "UObject/TextProperty.h"
 #include "UObject/UnrealType.h"
+#include "RestorationEvents/ApplySnapshotPropertiesScope.h"
 #include "Util/PropertyUtil.h"
+#include "Util/SnapshotObjectUtil.h"
 
 namespace
 {
@@ -141,6 +143,12 @@ namespace
 
 void FApplySnapshotDataArchiveV2::ApplyToExistingEditorWorldObject(FObjectSnapshotData& InObjectData, FWorldSnapshotData& InSharedData, UObject* InOriginalObject, UObject* InDeserializedVersion, const FPropertySelectionMap& InSelectionMapForResolvingSubobjects, const FPropertySelection& InSelectionSet)
 {
+	if (InSelectionSet.IsEmpty())
+	{
+		return;
+	}
+	
+	const FApplySnapshotPropertiesScope NotifySnapshotListeners({ InOriginalObject, InSelectionMapForResolvingSubobjects, &InSelectionSet, true });
 #if WITH_EDITOR
 	InOriginalObject->PreEditChange(nullptr);
 	ON_SCOPE_EXIT
@@ -162,6 +170,7 @@ void FApplySnapshotDataArchiveV2::ApplyToExistingEditorWorldObject(FObjectSnapsh
 
 void FApplySnapshotDataArchiveV2::ApplyToRecreatedEditorWorldObject(FObjectSnapshotData& InObjectData, FWorldSnapshotData& InSharedData, UObject* InOriginalObject, UObject* InDeserializedVersion, const FPropertySelectionMap& InSelectionMapForResolvingSubobjects)
 {
+	const FApplySnapshotPropertiesScope NotifySnapshotListeners({ InOriginalObject, InSelectionMapForResolvingSubobjects, {}, true });
 #if WITH_EDITOR
 	InOriginalObject->PreEditChange(nullptr);
 	ON_SCOPE_EXIT
@@ -184,7 +193,7 @@ bool FApplySnapshotDataArchiveV2::ShouldSkipProperty(const FProperty* InProperty
 	
 	if (!bShouldSkipProperty && !ShouldSerializeAllProperties())
 	{
-		const bool bIsAllowed = SelectionSet.GetValue()->ShouldSerializeProperty(GetSerializedPropertyChain(), InProperty) || CastField<FObjectPropertyBase>(InProperty);  
+		const bool bIsAllowed = SelectionSet.GetValue()->ShouldSerializeProperty(GetSerializedPropertyChain(), InProperty);  
 		bShouldSkipProperty = !bIsAllowed;
 	}
 	
@@ -201,7 +210,7 @@ void FApplySnapshotDataArchiveV2::PushSerializedProperty(FProperty* InProperty,
 
 UObject* FApplySnapshotDataArchiveV2::ResolveObjectDependency(int32 ObjectIndex) const
 {
-	return GetSharedData().ResolveObjectDependencyForEditorWorld(ObjectIndex, SelectionMapForResolvingSubobjects);
+	return SnapshotUtil::Object::ResolveObjectDependencyForEditorWorld(GetSharedData(), ObjectIndex, GetLocalizationNamespace(), SelectionMapForResolvingSubobjects);
 }
 
 FApplySnapshotDataArchiveV2::FApplySnapshotDataArchiveV2(FObjectSnapshotData& InObjectData, FWorldSnapshotData& InSharedData, UObject* InOriginalObject, const FPropertySelectionMap& InSelectionMapForResolvingSubobjects, TOptional<const FPropertySelection*> InSelectionSet)
diff --git a/Engine/Plugins/VirtualProduction/LevelSnapshots/Source/LevelSnapshots/Private/Archive/ClassDefaults/BaseClassDefaultArchive.cpp b/Engine/Plugins/VirtualProduction/LevelSnapshots/Source/LevelSnapshots/Private/Archive/ClassDefaults/BaseClassDefaultArchive.cpp
--- a/Engine/Plugins/VirtualProduction/LevelSnapshots/Source/LevelSnapshots/Private/Archive/ClassDefaults/BaseClassDefaultArchive.cpp
+++ b/Engine/Plugins/VirtualProduction/LevelSnapshots/Source/LevelSnapshots/Private/Archive/ClassDefaults/BaseClassDefaultArchive.cpp
@@ -9,6 +9,7 @@
 #include "GameFramework/Actor.h"
 #include "Serialization/ArchiveSerializedPropertyChain.h"
 #include "Util/PropertyUtil.h"
+#include "Util/SnapshotObjectUtil.h"
 
 bool FBaseClassDefaultArchive::ShouldSkipProperty(const FProperty* InProperty) const
 {
@@ -20,7 +21,7 @@ bool FBaseClassDefaultArchive::ShouldSkipProperty(const FProperty* InProperty) c
 
 UObject* FBaseClassDefaultArchive::ResolveObjectDependency(int32 ObjectIndex) const
 {
-	return GetSharedData().ResolveObjectDependencyForClassDefaultObject(ObjectIndex);
+	return SnapshotUtil::Object::ResolveObjectDependencyForClassDefaultObject(GetSharedData(), ObjectIndex);
 }
 
 FBaseClassDefaultArchive::FBaseClassDefaultArchive(FObjectSnapshotData& InObjectData, FWorldSnapshotData& InSharedData, bool bIsLoading, UObject* InObjectToRestore)
diff --git a/Engine/Plugins/VirtualProduction/LevelSnapshots/Source/LevelSnapshots/Private/Archive/SnapshotArchive.cpp b/Engine/Plugins/VirtualProduction/LevelSnapshots/Source/LevelSnapshots/Private/Archive/SnapshotArchive.cpp
--- a/Engine/Plugins/VirtualProduction/LevelSnapshots/Source/LevelSnapshots/Private/Archive/SnapshotArchive.cpp
+++ b/Engine/Plugins/VirtualProduction/LevelSnapshots/Source/LevelSnapshots/Private/Archive/SnapshotArchive.cpp
@@ -7,19 +7,32 @@
 #include "SnapshotVersion.h"
 #include "WorldSnapshotData.h"
 
-#include "Components/ActorComponent.h"
-#include "GameFramework/Actor.h"
 #include "Internationalization/TextNamespaceUtil.h"
 #include "Internationalization/TextPackageNamespaceUtil.h"
-#include "Serialization/ArchiveSerializedPropertyChain.h"
 #include "UObject/ObjectMacros.h"
+#include "Util/SnapshotObjectUtil.h"
 
 void FSnapshotArchive::ApplyToSnapshotWorldObject(FObjectSnapshotData& InObjectData, FWorldSnapshotData& InSharedData, UObject* InObjectToRestore, UPackage* InLocalisationSnapshotPackage)
+{
+	ApplyToSnapshotWorldObject(
+		InObjectData,
+		InSharedData,
+		InObjectToRestore,
+#if USE_STABLE_LOCALIZATION_KEYS
+		TextNamespaceUtil::EnsurePackageNamespace(InLocalisationSnapshotPackage)
+#else
+		FString()
+#endif
+		);
+}
+
+void FSnapshotArchive::ApplyToSnapshotWorldObject(FObjectSnapshotData& InObjectData, FWorldSnapshotData& InSharedData, UObject* InObjectToRestore, const FString& InLocalisationNamespace)
 {
 	FSnapshotArchive Archive(InObjectData, InSharedData, true, InObjectToRestore);
 #if USE_STABLE_LOCALIZATION_KEYS
-	Archive.SetLocalizationNamespace(TextNamespaceUtil::EnsurePackageNamespace(InLocalisationSnapshotPackage));
+	Archive.SetLocalizationNamespace(InLocalisationNamespace);
 #endif
+
 	InObjectToRestore->Serialize(Archive);
 }
 
@@ -46,10 +59,8 @@ void FSnapshotArchive::Seek(int64 InPos)
 
 bool FSnapshotArchive::ShouldSkipProperty(const FProperty* InProperty) const
 {
-	const bool bIsPropertyUnsupported = InProperty->HasAnyPropertyFlags(ExcludedPropertyFlags)|| IsPropertyReferenceToSubobject(InProperty);
-	const bool bIsBlacklisted = FSnapshotRestorability::IsPropertyBlacklistedForCapture(InProperty);
-	const bool bIsWhitelisted = FSnapshotRestorability::IsPropertyWhitelistedForCapture(InProperty);
-	return bIsPropertyUnsupported || bIsBlacklisted || (!bIsPropertyUnsupported && bIsWhitelisted);
+	const bool bIsPropertyUnsupported = InProperty->HasAnyPropertyFlags(ExcludedPropertyFlags);
+	return bIsPropertyUnsupported || !FSnapshotRestorability::IsPropertyDesirableForCapture(InProperty);
 }
 
 FArchive& FSnapshotArchive::operator<<(FName& Value)
@@ -100,7 +111,7 @@ FArchive& FSnapshotArchive::operator<<(UObject*& Value)
 	}
 	else
 	{
-		int32 ReferenceIndex = SharedData.AddObjectDependency(Value);
+		int32 ReferenceIndex = SnapshotUtil::Object::AddObjectDependency(SharedData, Value);
 		*this << ReferenceIndex;
 	}
 	
@@ -140,26 +151,7 @@ void FSnapshotArchive::Serialize(void* Data, int64 Length)
 
 UObject* FSnapshotArchive::ResolveObjectDependency(int32 ObjectIndex) const
 {
-	return SharedData.ResolveObjectDependencyForSnapshotWorld(ObjectIndex);
-}
-
-bool FSnapshotArchive::IsPropertyReferenceToSubobject(const FProperty* InProperty) const
-{
-	const FObjectPropertyBase* ObjectProperty = CastField<FObjectPropertyBase>(InProperty);
-	if (!ObjectProperty)
-	{
-		return false;
-	}
-
-	const bool bIsMarkedAsSubobject = InProperty->HasAnyPropertyFlags(CPF_InstancedReference | CPF_ContainsInstancedReference | CPF_PersistentInstance);
-	const bool bIsComponentPtr = ObjectProperty->PropertyClass->IsChildOf(UActorComponent::StaticClass());
-	if (bIsMarkedAsSubobject || bIsComponentPtr)
-	{
-		return true;
-	}
-
-	// TODO: Walk GetSerializedPropertyChain to get the value ptr of ObjectProperty. Then check whether contained object IsIn(GetSerializedObject()). Extra complicated because sometimes values may be in arrays, set, and tmaps.
-	return false;
+	return SnapshotUtil::Object::ResolveObjectDependencyForSnapshotWorld(SharedData, ObjectIndex, GetLocalizationNamespace());
 }
 
 FSnapshotArchive::FSnapshotArchive(FObjectSnapshotData& InObjectData, FWorldSnapshotData& InSharedData, bool bIsLoading, UObject* InSerializedObject)
diff --git a/Engine/Plugins/VirtualProduction/LevelSnapshots/Source/LevelSnapshots/Private/Archive/SnapshotArchive.h b/Engine/Plugins/VirtualProduction/LevelSnapshots/Source/LevelSnapshots/Private/Archive/SnapshotArchive.h
--- a/Engine/Plugins/VirtualProduction/LevelSnapshots/Source/LevelSnapshots/Private/Archive/SnapshotArchive.h
+++ b/Engine/Plugins/VirtualProduction/LevelSnapshots/Source/LevelSnapshots/Private/Archive/SnapshotArchive.h
@@ -20,6 +20,7 @@ class FSnapshotArchive : public FArchiveUObject
 public:
 
 	static void ApplyToSnapshotWorldObject(FObjectSnapshotData& InObjectData, FWorldSnapshotData& InSharedData, UObject* InObjectToRestore, UPackage* InLocalisationSnapshotPackage);
+	static void ApplyToSnapshotWorldObject(FObjectSnapshotData& InObjectData, FWorldSnapshotData& InSharedData, UObject* InObjectToRestore, const FString& InLocalisationNamespace);
 
 	//~ Begin FArchive Interface
 	virtual FString GetArchiveName() const override;
@@ -37,8 +38,6 @@ protected:
 	
 	/* Allocates and serializes an object dependency, or gets the object, if it already exists. */
 	virtual UObject* ResolveObjectDependency(int32 ObjectIndex) const;
-
-	bool IsPropertyReferenceToSubobject(const FProperty* InProperty) const;
 	
 	FWorldSnapshotData& GetSharedData() const { return SharedData;}
 	UObject* GetSerializedObject() const { return SerializedObject; }
diff --git a/Engine/Plugins/VirtualProduction/LevelSnapshots/Source/LevelSnapshots/Private/Data/ActorSnapshotData.cpp b/Engine/Plugins/VirtualProduction/LevelSnapshots/Source/LevelSnapshots/Private/Data/ActorSnapshotData.cpp
--- a/Engine/Plugins/VirtualProduction/LevelSnapshots/Source/LevelSnapshots/Private/Data/ActorSnapshotData.cpp
+++ b/Engine/Plugins/VirtualProduction/LevelSnapshots/Source/LevelSnapshots/Private/Data/ActorSnapshotData.cpp
@@ -3,16 +3,212 @@
 #include "ActorSnapshotData.h"
 
 #include "ApplySnapshotDataArchiveV2.h"
+#include "CustomSerialization/CustomObjectSerializationWrapper.h"
+#include "CustomSerialization/CustomSerializationDataManager.h"
 #include "LevelSnapshotsLog.h"
 #include "PropertySelectionMap.h"
+#include "SnapshotCustomVersion.h"
+#include "SnapshotRestorability.h"
 #include "TakeWorldObjectSnapshotArchive.h"
 #include "WorldSnapshotData.h"
-#include "CustomSerialization/CustomObjectSerializationWrapper.h"
-#include "CustomSerialization/CustomSerializationDataManager.h"
+#include "Util/SnapshotObjectUtil.h"
+#include "Util/SnapshotUtil.h"
 
 #include "Engine/World.h"
 #include "GameFramework/Actor.h"
+#if USE_STABLE_LOCALIZATION_KEYS
+#include "Internationalization/TextPackageNamespaceUtil.h"
+#endif
+#include "LevelEditor.h"
+#include "UnrealEdGlobals.h"
+#include "Editor/UnrealEdEngine.h"
+#include "RestorationEvents/ApplySnapshotToActorScope.h"
+#include "RestorationEvents/RecreateComponentScope.h"
+#include "RestorationEvents/RemoveComponentScope.h"
+#include "Templates/NonNullPointer.h"
 #include "UObject/Package.h"
+#include "UObject/Script.h"
+
+namespace
+{	
+	void FindOrRecreateSavedComponents(const FActorSnapshotData& SnapshotData, FWorldSnapshotData& WorldData, AActor* Actor)
+	{
+		TInlineComponentArray<UActorComponent*> DefaultSnapshotComps(Actor);
+		for (auto CompIt = SnapshotData.ComponentData.CreateConstIterator(); CompIt; ++CompIt)
+		{
+			const int32 ReferenceIndex = CompIt->Key;
+			FSubobjectSnapshotData* SubobjectData = WorldData.Subobjects.Find(ReferenceIndex); 
+			if (!ensure(SubobjectData))
+			{
+				continue;
+			}
+			
+			const FSoftObjectPath& ComponentPath = WorldData.SerializedObjectReferences[ReferenceIndex];
+			const FString ComponentName = SnapshotUtil::ExtractLastSubobjectName(ComponentPath);
+			UActorComponent* const* PossibleCounterpart = DefaultSnapshotComps.FindByPredicate([&ComponentName](UActorComponent* Component)
+			{
+				return Component->GetName() == ComponentName;
+			});
+			
+			UActorComponent* CounterpartComponent = PossibleCounterpart ? *PossibleCounterpart : nullptr;
+			if (!CounterpartComponent)
+			{
+				CounterpartComponent = NewObject<UActorComponent>(Actor, SubobjectData->Class.ResolveClass(), FName(*ComponentName));
+				SubobjectData->SnapshotObject = CounterpartComponent;
+			}
+			
+			// UActorComponent::PostInitProperties implicitly calls AddOwnedComponent but we have to manually add it to the other arrays
+			CounterpartComponent->CreationMethod = CompIt->Value.CreationMethod;
+			switch(CounterpartComponent->CreationMethod)
+			{
+				case EComponentCreationMethod::Instance:
+					Actor->AddInstanceComponent(CounterpartComponent);
+					break;
+				case EComponentCreationMethod::SimpleConstructionScript:
+					Actor->BlueprintCreatedComponents.Add(CounterpartComponent);
+					break;
+				case EComponentCreationMethod::UserConstructionScript:
+					checkf(false, TEXT("Component created in construction script currently unsupported"));
+					break;
+					
+				case EComponentCreationMethod::Native:
+				default:
+					break;
+			}
+		}
+	}
+
+	TPair<int32, EComponentCreationMethod> FindMatchingSubobjectIndex(const FActorSnapshotData& SnapshotData, const FWorldSnapshotData& WorldData, const FString& ComponentName)
+	{
+		for (auto CompIt = SnapshotData.ComponentData.CreateConstIterator(); CompIt; ++CompIt)
+		{
+			const int32 ReferenceIndex = CompIt->Key;
+			const FSoftObjectPath& SavedComponentPath = WorldData.SerializedObjectReferences[ReferenceIndex];
+			const FString SavedComponentName = SnapshotUtil::ExtractLastSubobjectName(SavedComponentPath);
+			if (SavedComponentName.Equals(ComponentName))
+			{
+				return TPair<int32, EComponentCreationMethod>(ReferenceIndex, CompIt->Value.CreationMethod);
+			}
+		}
+
+		return TPair<int32, EComponentCreationMethod>(INDEX_NONE, EComponentCreationMethod::Instance);
+	}
+
+	void UpdateDetailsViewAfterUpdatingComponents()
+	{
+#if WITH_EDITOR
+		FLevelEditorModule& LevelEditor = FModuleManager::LoadModuleChecked<FLevelEditorModule>("LevelEditor");
+		LevelEditor.BroadcastComponentsEdited();
+#endif
+	}
+	
+	void RemovedNewComponents(const FAddedAndRemovedComponentInfo& ComponentSelection)
+	{
+		for (const TWeakObjectPtr<UActorComponent>& ComponentToRemove : ComponentSelection.EditorWorldComponentsToRemove)
+		{
+			// Evil user might have manually deleted the component between filtering and applying the snapshot
+			if (ComponentToRemove.IsValid())
+			{
+				const FRemoveComponentScope NotifySnapshotListeners(ComponentToRemove.Get());
+				ComponentToRemove->Modify();
+				ComponentToRemove->DestroyComponent();
+			}
+		}
+	}
+	
+	void RecreateMissingComponents(const FActorSnapshotData& SnapshotData, FWorldSnapshotData& WorldData, AActor* Actor, const FAddedAndRemovedComponentInfo& ComponentSelection, const FPropertySelectionMap& SelectionMap, UPackage* LocalisationSnapshotPackage)
+	{
+		struct FDeferredComponentSerializationTask
+		{
+			FSubobjectSnapshotData* SubobjectData;
+			UActorComponent* Snapshot;
+			UActorComponent* Original;
+
+			FDeferredComponentSerializationTask(FSubobjectSnapshotData* SubobjectData, UActorComponent* Snapshot, UActorComponent* Original)
+				:
+				SubobjectData(SubobjectData),
+				Snapshot(Snapshot),
+				Original(Original)
+			{}
+		};
+		
+		// Recreate missing components
+		TArray<FDeferredComponentSerializationTask> DeferredSerializationsTasks;
+		for (const TWeakObjectPtr<UActorComponent>& ComponentToRestore : ComponentSelection.SnapshotComponentsToAdd)
+		{
+			const FName ComponentName = ComponentToRestore->GetFName();
+			const TPair<int32, EComponentCreationMethod> ComponentInfo = FindMatchingSubobjectIndex(SnapshotData, WorldData, ComponentName.ToString());
+			const int32 ReferenceIndex = ComponentInfo.Key;
+			if (ReferenceIndex == INDEX_NONE)
+			{
+				continue;
+			}
+			FSubobjectSnapshotData* SubobjectData = WorldData.Subobjects.Find(ReferenceIndex); 
+			if (!ensure(SubobjectData))
+			{
+				continue;
+			}
+			
+			UClass* ComponentClass = SubobjectData->Class.TryLoadClass<UActorComponent>(); 
+			const FRecreateComponentScope NotifySnapshotListeners(*SubobjectData, Actor, ComponentName, ComponentClass, ComponentInfo.Value);
+			UActorComponent* Component = NewObject<UActorComponent>(Actor, ComponentClass, ComponentName);
+			if (!Component)
+			{
+				UE_LOG(LogLevelSnapshots, Error, TEXT("Failed to recreate component called %s of class %s. Was the class removed?"), *ComponentName.ToString(), *SubobjectData->Class.ToString());
+				continue;
+			}
+			SubobjectData->EditorObject = Component;
+			
+			// UActorComponent::PostInitProperties implicitly calls AddOwnedComponent but we have to manually add it to the other arrays
+			Component->CreationMethod = ComponentInfo.Value;
+			switch(Component->CreationMethod)
+			{
+			case EComponentCreationMethod::Instance:
+				Actor->AddInstanceComponent(Component);
+				break;
+			case EComponentCreationMethod::SimpleConstructionScript:
+				Actor->BlueprintCreatedComponents.Add(Component);
+				break;
+			case EComponentCreationMethod::UserConstructionScript:
+				checkf(false, TEXT("Component created in construction script currently unsupported"));
+				break;
+					
+			case EComponentCreationMethod::Native:
+			default:
+				break;
+			}
+			
+			Component->RegisterComponent();
+			// Defer task until all components allocated: components may have references to each other, e.g. AttachParent
+			DeferredSerializationsTasks.Add({ SubobjectData, ComponentToRestore.Get(), Component });
+		}
+
+		// Write all saved data into the recreated component
+		for (const FDeferredComponentSerializationTask& Task : DeferredSerializationsTasks)
+		{
+			const FRestoreObjectScope FinishRestore = FCustomObjectSerializationWrapper::PreSubobjectRestore_EditorWorld(Task.Snapshot, Task.Original, WorldData, SelectionMap, LocalisationSnapshotPackage);
+			FApplySnapshotDataArchiveV2::ApplyToRecreatedEditorWorldObject(*Task.SubobjectData, WorldData, Task.Original, Task.Snapshot, SelectionMap); 
+		}
+	}
+
+	void AddAndRemoveSelectedComponentsForRestore(
+		const FActorSnapshotData& SnapshotData,
+		FWorldSnapshotData& WorldData,
+		AActor* Actor,
+		const FAddedAndRemovedComponentInfo& ComponentSelection,
+		const FPropertySelectionMap& SelectionMap,
+		UPackage* LocalisationSnapshotPackage)
+	{
+		RemovedNewComponents(ComponentSelection);
+		RecreateMissingComponents(SnapshotData, WorldData, Actor, ComponentSelection, SelectionMap, LocalisationSnapshotPackage);
+
+		const bool bChangedHierarchy = ComponentSelection.SnapshotComponentsToAdd.Num() > 0 || ComponentSelection.EditorWorldComponentsToRemove.Num() > 0;
+		if (bChangedHierarchy)
+		{
+			UpdateDetailsViewAfterUpdatingComponents();
+		}
+	}
+}
 
 FActorSnapshotData FActorSnapshotData::SnapshotActor(AActor* OriginalActor, FWorldSnapshotData& WorldData)
 {
@@ -30,10 +226,15 @@ FActorSnapshotData FActorSnapshotData::SnapshotActor(AActor* OriginalActor, FWor
 	OriginalActor->GetComponents(Components);
 	for (UActorComponent* Comp : Components)
 	{
+		if (!FSnapshotRestorability::IsComponentDesirableForCapture(Comp))
+		{
+			continue;
+		}
+		
 		TOptional<FComponentSnapshotData> SerializedComponentData = FComponentSnapshotData::SnapshotComponent(Comp, WorldData);
 		if (SerializedComponentData)
 		{
-			const int32 ComponentIndex = WorldData.AddSubobjectDependency(Comp);
+			const int32 ComponentIndex = SnapshotUtil::Object::AddObjectDependency(WorldData, Comp);
 			Result.ComponentData.Add(ComponentIndex, *SerializedComponentData);
 			// If external modules registered for custom serialisation, trigger their callbacks
 			FCustomObjectSerializationWrapper::TakeSnapshotForSubobject(Comp, WorldData);
@@ -43,46 +244,6 @@ FActorSnapshotData FActorSnapshotData::SnapshotActor(AActor* OriginalActor, FWor
 	return Result;
 }
 
-void FActorSnapshotData::DeserializeIntoExistingWorldActor(UWorld* SnapshotWorld, AActor* OriginalActor, FWorldSnapshotData& WorldData, UPackage* InLocalisationSnapshotPackage, const FPropertySelectionMap& SelectedProperties)
-{
-	auto DeserializeActor = [this, &WorldData, InLocalisationSnapshotPackage, &SelectedProperties](AActor* OriginalActor, AActor* DeserializedActor)
-	{
-		const FRestoreObjectScope FinishRestore = FCustomObjectSerializationWrapper::PreActorRestore_EditorWorld(OriginalActor, CustomActorSerializationData, WorldData, SelectedProperties, InLocalisationSnapshotPackage);
-		const FPropertySelection* ActorPropertySelection = SelectedProperties.GetSelectedProperties(OriginalActor);
-		if (ActorPropertySelection)
-		{
-			FApplySnapshotDataArchiveV2::ApplyToExistingEditorWorldObject(SerializedActorData, WorldData, OriginalActor, DeserializedActor, SelectedProperties, *ActorPropertySelection);
-		}
-	};
-	auto DeserializeComponent = [OriginalActor, InLocalisationSnapshotPackage, &SelectedProperties, &WorldData](FObjectSnapshotData& SerializedCompData, FComponentSnapshotData& CompData, UActorComponent* Original, UActorComponent* Deserialized)
-	{
-		const FRestoreObjectScope FinishRestore = FCustomObjectSerializationWrapper::PreSubobjectRestore_EditorWorld(Deserialized, Original, WorldData, SelectedProperties, InLocalisationSnapshotPackage);
-		const FPropertySelection* ComponentSelectedProperties = SelectedProperties.GetSelectedProperties(Original);
-		if (ComponentSelectedProperties)
-		{		
-			FApplySnapshotDataArchiveV2::ApplyToExistingEditorWorldObject(SerializedCompData, WorldData, Original, Deserialized, SelectedProperties, *ComponentSelectedProperties);
-		};
-	};
-
-	DeserializeIntoWorldActor(SnapshotWorld, OriginalActor, WorldData, InLocalisationSnapshotPackage, DeserializeActor, DeserializeComponent);
-}
-
-void FActorSnapshotData::DeserializeIntoRecreatedEditorWorldActor(UWorld* SnapshotWorld, AActor* OriginalActor, FWorldSnapshotData& WorldData, UPackage* InLocalisationSnapshotPackage, const FPropertySelectionMap& SelectedProperties)
-{
-	auto DeserializeActor = [this, &WorldData, InLocalisationSnapshotPackage, &SelectedProperties](AActor* OriginalActor, AActor* DeserializedActor)
-	{
-		const FRestoreObjectScope FinishRestore = FCustomObjectSerializationWrapper::PreActorRestore_EditorWorld(OriginalActor, CustomActorSerializationData, WorldData, SelectedProperties, InLocalisationSnapshotPackage);
-		FApplySnapshotDataArchiveV2::ApplyToRecreatedEditorWorldObject(SerializedActorData, WorldData, OriginalActor, DeserializedActor, SelectedProperties);
-	};
-	auto DeserializeComponent = [OriginalActor, InLocalisationSnapshotPackage, &WorldData, &SelectedProperties](FObjectSnapshotData& SerializedCompData, FComponentSnapshotData& CompData, UActorComponent* Original, UActorComponent* Deserialized)
-	{
-		const FRestoreObjectScope FinishRestore = FCustomObjectSerializationWrapper::PreSubobjectRestore_EditorWorld(Deserialized, Original, WorldData, SelectedProperties, InLocalisationSnapshotPackage);
-		FApplySnapshotDataArchiveV2::ApplyToRecreatedEditorWorldObject(SerializedCompData, WorldData, Original, Deserialized, SelectedProperties); 
-	};
-	
-	DeserializeIntoWorldActor(SnapshotWorld, OriginalActor, WorldData, InLocalisationSnapshotPackage, DeserializeActor, DeserializeComponent);
-}
-
 TOptional<AActor*> FActorSnapshotData::GetPreallocatedIfValidButDoNotAllocate() const
 {
 	return CachedSnapshotActor.IsValid() ? CachedSnapshotActor.Get() : TOptional<AActor*>();
@@ -90,6 +251,8 @@ TOptional<AActor*> FActorSnapshotData::GetPreallocatedIfValidButDoNotAllocate()
 
 TOptional<AActor*> FActorSnapshotData::GetPreallocated(UWorld* SnapshotWorld, FWorldSnapshotData& WorldData) const
 {
+	SCOPED_SNAPSHOT_CORE_TRACE(GetPreallocated);
+	
 	if (!CachedSnapshotActor.IsValid())
 	{
 		const FSoftClassPath SoftClassPath(ActorClass);
@@ -100,7 +263,6 @@ TOptional<AActor*> FActorSnapshotData::GetPreallocated(UWorld* SnapshotWorld, FW
 			return {};
 		}
 		
-		// TODO: Maybe there is a faster way than calling SpawnActor...
 		FActorSpawnParameters SpawnParams;
 		SpawnParams.Template = Cast<AActor>(WorldData.GetClassDefault(TargetClass));
 		SpawnParams.NameMode = FActorSpawnParameters::ESpawnActorNameMode::Requested;
@@ -114,21 +276,26 @@ TOptional<AActor*> FActorSnapshotData::GetPreallocated(UWorld* SnapshotWorld, FW
 		{
 			CachedSnapshotActor = SnapshotWorld->SpawnActor<AActor>(TargetClass, SpawnParams);
 		}
-	}
-
+		
+		if (!ensureMsgf(CachedSnapshotActor.IsValid(), TEXT("Failed to spawn actor of class '%s'"), *ActorClass.ToString()))
+		{
+			return {};
+		}
+		
 #if WITH_EDITOR
-	// Hide this actor so external systems can see that this components should not render, i.e. make USceneComponent::ShouldRender return false
-	if (ensureMsgf(CachedSnapshotActor.IsValid(), TEXT("Failed to spawn actor of class '%s'"), *ActorClass.ToString()))
-	{
+		// Hide this actor so external systems can see that this components should not render, i.e. make USceneComponent::ShouldRender return false
 		CachedSnapshotActor->SetIsTemporarilyHiddenInEditor(true);
-	}
 #endif
-	
+		FindOrRecreateSavedComponents(*this, WorldData, CachedSnapshotActor.Get());
+	}
+
 	return CachedSnapshotActor.IsValid() ? TOptional<AActor*>(CachedSnapshotActor.Get()) : TOptional<AActor*>();
 }
 
-TOptional<AActor*> FActorSnapshotData::GetDeserialized(UWorld* SnapshotWorld, FWorldSnapshotData& WorldData, UPackage* InLocalisationSnapshotPackage)
+TOptional<AActor*> FActorSnapshotData::GetDeserialized(UWorld* SnapshotWorld, FWorldSnapshotData& WorldData, const FSoftObjectPath& OriginalActorPath, UPackage* InLocalisationSnapshotPackage)
 {
+	SCOPED_SNAPSHOT_CORE_TRACE(GetDeserialized);
+	
 	if (bReceivedSerialisation && CachedSnapshotActor.IsValid())
 	{
 		return CachedSnapshotActor.Get();
@@ -140,7 +307,7 @@ TOptional<AActor*> FActorSnapshotData::GetDeserialized(UWorld* SnapshotWorld, FW
 		return {};
 	}
 	bReceivedSerialisation = true;
-	
+
 	AActor* PreallocatedActor = Preallocated.GetValue();
 	{
 		const FRestoreObjectScope FinishRestore = FCustomObjectSerializationWrapper::PreActorRestore_SnapshotWorld(PreallocatedActor, CustomActorSerializationData, WorldData, InLocalisationSnapshotPackage);
@@ -149,18 +316,20 @@ TOptional<AActor*> FActorSnapshotData::GetDeserialized(UWorld* SnapshotWorld, FW
 
 	DeserializeComponents(PreallocatedActor, WorldData,
 		[&WorldData, InLocalisationSnapshotPackage](
-			FObjectSnapshotData& SerializedCompData,
+			FSubobjectSnapshotData& SerializedCompData,
 			FComponentSnapshotData& CompData,
 			UActorComponent* Comp,
 			const FSoftObjectPath& OriginalComponentPath,
 			FWorldSnapshotData& SharedData)
 		{
+			SerializedCompData.SnapshotObject = Comp;
+			
 			const FRestoreObjectScope FinishRestore = FCustomObjectSerializationWrapper::PreSubobjectRestore_SnapshotWorld(Comp, OriginalComponentPath, WorldData, InLocalisationSnapshotPackage);
 			CompData.DeserializeIntoTransient(SerializedCompData, Comp, SharedData, InLocalisationSnapshotPackage);
 		}
 	);
 
-	PreallocatedActor->UpdateComponentTransforms();
+	DeserializeSubobjectsForSnapshotActor(PreallocatedActor, WorldData, InLocalisationSnapshotPackage);
 #if WITH_EDITOR
 	// Hide this actor so external systems can see that this components should not render, i.e. make USceneComponent::ShouldRender return false
 	if (!ensureMsgf(PreallocatedActor->IsTemporarilyHiddenInEditor(), TEXT("Transient property bHiddenEdTemporary was set to false by serializer. This should not happen. Investigate.")))
@@ -168,50 +337,190 @@ TOptional<AActor*> FActorSnapshotData::GetDeserialized(UWorld* SnapshotWorld, FW
 		CachedSnapshotActor->SetIsTemporarilyHiddenInEditor(true);
 	}
 #endif
+
+	PostSerializeSnapshotActor(PreallocatedActor, WorldData, OriginalActorPath, InLocalisationSnapshotPackage);
 	
+	PreallocatedActor->UpdateComponentTransforms();
+	{
+		// GAllowActorScriptExecutionInEditor must be temporarily true so call to UserConstructionScript isn't skipped
+		FEditorScriptExecutionGuard AllowConstructionScript;
+		PreallocatedActor->UserConstructionScript();
+	}
 	return Preallocated;
 }
 
-void FActorSnapshotData::DeserializeIntoWorldActor(UWorld* SnapshotWorld, AActor* OriginalActor, FWorldSnapshotData& WorldData, UPackage* InLocalisationSnapshotPackage, FSerializeActor SerializeActor, FSerializeComponent SerializeComponent)
+void FActorSnapshotData::PostSerializeSnapshotActor(AActor* SnapshotActor, FWorldSnapshotData& WorldData, const FSoftObjectPath& OriginalActorPath, UPackage* InLocalisationSnapshotPackage) const
 {
-	const TOptional<AActor*> Deserialized = GetDeserialized(SnapshotWorld, WorldData, InLocalisationSnapshotPackage);
+	struct Local
+	{
+		static void RecreateRootComponentIfInstanced(AActor* OriginalActor, AActor* SnapshotActor)
+		{
+			USceneComponent* OriginalRoot = OriginalActor->GetRootComponent();
+			
+			const bool bSnapshotIsMissingRoot = SnapshotActor->GetRootComponent() == nullptr;
+			// Some actors, e.g. exact instances of AActor, have no components and the editor creates an instanced root component
+			const bool bSnapshotWasSavedWithoutRootComponentData = OriginalRoot && OriginalRoot->CreationMethod == EComponentCreationMethod::Instance;
+			if (bSnapshotIsMissingRoot && bSnapshotWasSavedWithoutRootComponentData)
+			{
+				USceneComponent* SnapshotRoot = DuplicateObject<USceneComponent>(OriginalRoot, SnapshotActor, OriginalRoot->GetFName());
+				UE_CLOG(SnapshotRoot == nullptr, LogLevelSnapshots, Warning, TEXT("Failed to recreate instanced root component for %s's snapshot counterpart"), *OriginalActor->GetPathName());
+				if (SnapshotRoot)
+				{
+					SnapshotActor->SetRootComponent(
+						SnapshotRoot
+					);
+				}
+			}
+		}
+		
+		static void HandleAttachParentNotSaved(AActor* OriginalActor, AActor* SnapshotActor, FWorldSnapshotData& WorldData, UPackage* InLocalisationSnapshotPackage)
+		{
+			USceneComponent* RootComponent = OriginalActor->GetRootComponent();
+			// WorldSettings does not have a root component...
+			if (!RootComponent)
+			{
+				return;
+			}
+			
+			USceneComponent* WorldAttachParent = RootComponent->GetAttachParent();
+			if (!WorldAttachParent)
+			{
+				return;
+			}
+			const TOptional<AActor*> SnapshotAttachParent = WorldData.GetDeserializedActor(WorldAttachParent->GetOwner(), InLocalisationSnapshotPackage);
+			if (!SnapshotAttachParent)
+			{
+				return;
+			}
+
+			const TInlineComponentArray<USceneComponent*> ComponentArray(*SnapshotAttachParent);
+			USceneComponent* const* SnapshotAttachComponent = ComponentArray.FindByPredicate([WorldAttachParent](USceneComponent* Comp)
+			{
+				return Comp->GetFName() == WorldAttachParent->GetFName();
+			});
+
+			if (SnapshotAttachComponent)
+			{
+				SnapshotActor->AttachToComponent(
+					*SnapshotAttachComponent,
+					// Do not change property values: anything other than KeepRelativeTransform will recompute
+					FAttachmentTransformRules::KeepRelativeTransform,
+					RootComponent->GetAttachSocketName()
+					);
+			}
+			else
+			{
+				UE_LOG(LogLevelSnapshots, Warning, TEXT("Failed to fix up AttachParent for %s's snapshot counterpart"), *OriginalActor->GetPathName());
+			}
+		}
+	};
+	
+	const int32 SavedVersion = WorldData.GetSnapshotVersionInfo().GetSnapshotCustomVersion();
+	AActor* OriginalActor = Cast<AActor>(OriginalActorPath.ResolveObject());
+
+	// USceneComponent::AttachParent was not yet saved.
+	if (SavedVersion < FSnapshotCustomVersion::SubobjectSupport
+		&& OriginalActor)
+	{
+		Local::RecreateRootComponentIfInstanced(OriginalActor, SnapshotActor);
+		// To avoid lots of actors being unparented by accident, we set the snapshot actor's attach parent to the equivalent of the editor world
+		Local::HandleAttachParentNotSaved(OriginalActor, SnapshotActor, WorldData, InLocalisationSnapshotPackage);
+	}
+}
+
+void FActorSnapshotData::DeserializeIntoExistingWorldActor(UWorld* SnapshotWorld, AActor* OriginalActor, FWorldSnapshotData& WorldData, UPackage* InLocalisationSnapshotPackage, const FPropertySelectionMap& SelectedProperties)
+{
+	auto DeserializeActor = [this, &WorldData, InLocalisationSnapshotPackage, &SelectedProperties](AActor* OriginalActor, AActor* DeserializedActor)
+	{
+		const FRestoreObjectScope FinishRestore = FCustomObjectSerializationWrapper::PreActorRestore_EditorWorld(OriginalActor, CustomActorSerializationData, WorldData, SelectedProperties, InLocalisationSnapshotPackage);
+		const FPropertySelection* ActorPropertySelection = SelectedProperties.GetObjectSelection(OriginalActor).GetPropertySelection();
+		if (ActorPropertySelection)
+		{
+			FApplySnapshotDataArchiveV2::ApplyToExistingEditorWorldObject(SerializedActorData, WorldData, OriginalActor, DeserializedActor, SelectedProperties, *ActorPropertySelection);
+		}
+	};
+	auto DeserializeComponent = [OriginalActor, InLocalisationSnapshotPackage, &SelectedProperties, &WorldData](FSubobjectSnapshotData& SerializedCompData, FComponentSnapshotData& CompData, UActorComponent* Original, UActorComponent* Deserialized)
+	{
+		const FRestoreObjectScope FinishRestore = FCustomObjectSerializationWrapper::PreSubobjectRestore_EditorWorld(Deserialized, Original, WorldData, SelectedProperties, InLocalisationSnapshotPackage);
+		const FPropertySelection* ComponentSelectedProperties = SelectedProperties.GetObjectSelection(Original).GetPropertySelection();
+		if (ComponentSelectedProperties)
+		{		
+			FApplySnapshotDataArchiveV2::ApplyToExistingEditorWorldObject(SerializedCompData, WorldData, Original, Deserialized, SelectedProperties, *ComponentSelectedProperties);
+		};
+	};
+
+	const bool bWasRecreated = false;
+	const FApplySnapshotToActorScope NotifyExternalListeners({ OriginalActor, SelectedProperties, bWasRecreated });
+	
+	if (const FAddedAndRemovedComponentInfo* ComponentSelection = SelectedProperties.GetObjectSelection(OriginalActor).GetComponentSelection())
+	{
+		AddAndRemoveSelectedComponentsForRestore(*this, WorldData, OriginalActor, *ComponentSelection, SelectedProperties, InLocalisationSnapshotPackage);
+	}
+	DeserializeIntoEditorWorldActor(SnapshotWorld, OriginalActor, WorldData, InLocalisationSnapshotPackage, DeserializeActor, DeserializeComponent);
+}
+
+void FActorSnapshotData::DeserializeIntoRecreatedEditorWorldActor(UWorld* SnapshotWorld, AActor* OriginalActor, FWorldSnapshotData& WorldData, UPackage* InLocalisationSnapshotPackage, const FPropertySelectionMap& SelectedProperties)
+{
+	auto DeserializeActor = [this, &WorldData, InLocalisationSnapshotPackage, &SelectedProperties](AActor* OriginalActor, AActor* DeserializedActor)
+	{
+		const FRestoreObjectScope FinishRestore = FCustomObjectSerializationWrapper::PreActorRestore_EditorWorld(OriginalActor, CustomActorSerializationData, WorldData, SelectedProperties, InLocalisationSnapshotPackage);
+		FApplySnapshotDataArchiveV2::ApplyToRecreatedEditorWorldObject(SerializedActorData, WorldData, OriginalActor, DeserializedActor, SelectedProperties);
+	};
+	auto DeserializeComponent = [OriginalActor, InLocalisationSnapshotPackage, &WorldData, &SelectedProperties](FSubobjectSnapshotData& SerializedCompData, FComponentSnapshotData& CompData, UActorComponent* Original, UActorComponent* Deserialized)
+	{
+		const FRestoreObjectScope FinishRestore = FCustomObjectSerializationWrapper::PreSubobjectRestore_EditorWorld(Deserialized, Original, WorldData, SelectedProperties, InLocalisationSnapshotPackage);
+		FApplySnapshotDataArchiveV2::ApplyToRecreatedEditorWorldObject(SerializedCompData, WorldData, Original, Deserialized, SelectedProperties); 
+	};
+
+	const bool bWasRecreated = true;
+	const FApplySnapshotToActorScope NotifyExternalListeners({ OriginalActor, SelectedProperties, bWasRecreated });
+	
+	FindOrRecreateSavedComponents(*this, WorldData, OriginalActor);
+	DeserializeIntoEditorWorldActor(SnapshotWorld, OriginalActor, WorldData, InLocalisationSnapshotPackage, DeserializeActor, DeserializeComponent);
+}
+
+void FActorSnapshotData::DeserializeIntoEditorWorldActor(UWorld* SnapshotWorld, AActor* OriginalActor, FWorldSnapshotData& WorldData, UPackage* InLocalisationSnapshotPackage, FSerializeActor SerializeActor, FSerializeComponent SerializeComponent)
+{
+	const TOptional<AActor*> Deserialized = GetDeserialized(SnapshotWorld, WorldData, OriginalActor, InLocalisationSnapshotPackage);
 	if (!Deserialized)
 	{
 		UE_LOG(LogLevelSnapshots, Warning, TEXT("Failed to serialize into actor %s. Skipping..."), *OriginalActor->GetName());
 		return;
 	}
-
+	
 	SerializeActor(OriginalActor, *Deserialized);
 
 	TInlineComponentArray<UActorComponent*> DeserializedComponents;
 	Deserialized.GetValue()->GetComponents(DeserializedComponents);
 	DeserializeComponents(OriginalActor, WorldData,
 		[&SerializeComponent, &DeserializedComponents](
-			FObjectSnapshotData& SerializedCompData,
+			FSubobjectSnapshotData& SerializedCompData,
 			FComponentSnapshotData& CompData,
 			UActorComponent* Comp,
 			const FSoftObjectPath& OriginalComponenPath,
 			FWorldSnapshotData& SharedData
 			)
-    {
-        const FName OriginalCompName = Comp->GetFName();
-        UActorComponent** DeserializedCompCounterpart = DeserializedComponents.FindByPredicate([OriginalCompName](UActorComponent* Other)
-        {
-            return Other->GetFName() == OriginalCompName;
-        });
-        if (!DeserializedCompCounterpart)
-        {
-            UE_LOG(LogLevelSnapshots, Warning, TEXT("Failed to find component called %s on temp deserialized snapshot actor. Skipping component..."), *OriginalCompName.ToString())
-        	return;
-        }
-
-        SerializeComponent(SerializedCompData, CompData, Comp, *DeserializedCompCounterpart);
-		
-		// We may have modified render information, e.g. for lights we may have changed intensity or colour
-		// It may be more efficient to track whether we actually changed render state
-		Comp->MarkRenderStateDirty();
-    });
+	    {
+	        const FName OriginalCompName = Comp->GetFName();
+	        UActorComponent** DeserializedCompCounterpart = DeserializedComponents.FindByPredicate([OriginalCompName](UActorComponent* Other)
+	        {
+	            return Other->GetFName() == OriginalCompName;
+	        });
+	        if (!DeserializedCompCounterpart)
+	        {
+	            UE_LOG(LogLevelSnapshots, Warning, TEXT("Failed to find component called %s on temp deserialized snapshot actor. Skipping component..."), *OriginalCompName.ToString())
+        		return;
+	        }
 
+	        SerializeComponent(SerializedCompData, CompData, Comp, *DeserializedCompCounterpart);
+			
+			// We may have modified render information, e.g. for lights we may have changed intensity or colour
+			// It may be more efficient to track whether we actually changed render state
+			Comp->MarkRenderStateDirty();
+	    }
+	);
+
+	// Fixes up restored attachments
 	OriginalActor->UpdateComponentTransforms();
 }
 
@@ -223,14 +532,16 @@ void FActorSnapshotData::DeserializeComponents(
 {
 	for (auto CompIt = ComponentData.CreateIterator(); CompIt; ++CompIt)
 	{
-		// Instances and construction script are not supported 
-		if (!CompIt->Value.IsRestoreSupportedForSavedComponent())
+		const EComponentCreationMethod CreationMethod = CompIt->Value.CreationMethod;
+		
+		if (CreationMethod == EComponentCreationMethod::UserConstructionScript)	// Construction script components are not supported 
 		{
 			continue;
 		}
-		
-		const FSoftObjectPath& OriginalComponentPath = WorldData.SerializedObjectReferences[CompIt->Key];
-		FSubobjectSnapshotData& SnapshotData = WorldData.Subobjects[CompIt->Key];
+
+		const int32 SubobjectIndex = CompIt->Key;
+		const FSoftObjectPath& OriginalComponentPath = WorldData.SerializedObjectReferences[SubobjectIndex];
+		FSubobjectSnapshotData& SnapshotData = WorldData.Subobjects[SubobjectIndex];
 		
 		const FString& SubPath = OriginalComponentPath.GetSubPathString();
 		const int32 LastDot = SubPath.Find(TEXT("."), ESearchCase::IgnoreCase, ESearchDir::FromEnd);
@@ -253,3 +564,24 @@ void FActorSnapshotData::DeserializeComponents(
 		}
 	}
 }
+
+inline void FActorSnapshotData::DeserializeSubobjectsForSnapshotActor(AActor* IntoActor, FWorldSnapshotData& WorldData, UPackage* InLocalisationSnapshotPackage)
+{
+	for (const int32 SubobjectIndex : OwnedSubobjects)
+	{
+		check(!ComponentData.Contains(SubobjectIndex));
+
+#if USE_STABLE_LOCALIZATION_KEYS
+		const FString LocalisationNamespace = TextNamespaceUtil::EnsurePackageNamespace(InLocalisationSnapshotPackage);
+#else
+		const FString LocalisationNamespace;
+#endif
+
+		// Ensures the object is allocated and serialized into. Return value not needed.
+		SnapshotUtil::Object::ResolveObjectDependencyForSnapshotWorld(
+			WorldData,
+			SubobjectIndex,
+			LocalisationNamespace
+		);
+	}
+}
diff --git a/Engine/Plugins/VirtualProduction/LevelSnapshots/Source/LevelSnapshots/Private/Data/ComponentSnapshotData.cpp b/Engine/Plugins/VirtualProduction/LevelSnapshots/Source/LevelSnapshots/Private/Data/ComponentSnapshotData.cpp
--- a/Engine/Plugins/VirtualProduction/LevelSnapshots/Source/LevelSnapshots/Private/Data/ComponentSnapshotData.cpp
+++ b/Engine/Plugins/VirtualProduction/LevelSnapshots/Source/LevelSnapshots/Private/Data/ComponentSnapshotData.cpp
@@ -11,11 +11,6 @@
 
 TOptional<FComponentSnapshotData> FComponentSnapshotData::SnapshotComponent(UActorComponent* OriginalComponent, FWorldSnapshotData& WorldData)
 {
-	if (OriginalComponent->CreationMethod == EComponentCreationMethod::Instance)
-	{
-		UE_LOG(LogLevelSnapshots, Warning, TEXT("Components added via the Component section of the Actor's details panel are not supported (%s). Skipping..."), *OriginalComponent->GetPathName());
-		return {};
-	}
 	if (OriginalComponent->CreationMethod == EComponentCreationMethod::UserConstructionScript)
 	{
 		UE_LOG(LogLevelSnapshots, Warning, TEXT("Components created dynamically in the construction script are not supported (%s). Skipping..."), *OriginalComponent->GetPathName());
@@ -26,19 +21,8 @@ TOptional<FComponentSnapshotData> FComponentSnapshotData::SnapshotComponent(UAct
 	return Result;
 }
 
-bool FComponentSnapshotData::IsRestoreSupportedForSavedComponent() const
-{
-	return CreationMethod != EComponentCreationMethod::Instance && CreationMethod != EComponentCreationMethod::UserConstructionScript;
-}
-
 void FComponentSnapshotData::DeserializeIntoTransient(FObjectSnapshotData& SerializedComponentData, UActorComponent* ComponentToDeserializeInto, FWorldSnapshotData& WorldData, UPackage* InLocalisationSnapshotPackage)
 {
-	// TODO: Handle CreationMethod == Instance and UserConstructionScript here by re-creating the component if it does not exist
-	if (ComponentToDeserializeInto->CreationMethod == EComponentCreationMethod::Instance)
-	{
-		UE_LOG(LogLevelSnapshots, Warning, TEXT("Components added via the Component section of the Actor's details panel are not supported (%s). Skipping..."), *ComponentToDeserializeInto->GetPathName());
-		return;
-	}
 	if (ComponentToDeserializeInto->CreationMethod == EComponentCreationMethod::UserConstructionScript)
 	{
 		UE_LOG(LogLevelSnapshots, Warning, TEXT("Components created dynamically in the construction script are not supported (%s). Skipping..."), *ComponentToDeserializeInto->GetPathName());
diff --git a/Engine/Plugins/VirtualProduction/LevelSnapshots/Source/LevelSnapshots/Private/Data/CustomSerialization/CustomObjectSerializationWrapper.cpp b/Engine/Plugins/VirtualProduction/LevelSnapshots/Source/LevelSnapshots/Private/Data/CustomSerialization/CustomObjectSerializationWrapper.cpp
--- a/Engine/Plugins/VirtualProduction/LevelSnapshots/Source/LevelSnapshots/Private/Data/CustomSerialization/CustomObjectSerializationWrapper.cpp
+++ b/Engine/Plugins/VirtualProduction/LevelSnapshots/Source/LevelSnapshots/Private/Data/CustomSerialization/CustomObjectSerializationWrapper.cpp
@@ -96,7 +96,7 @@ namespace
 					continue;
 				}
 
-				if (const FPropertySelection* SelectedProperties = SelectionMap.GetSelectedProperties(EditorSubobject))
+				if (const FPropertySelection* SelectedProperties = SelectionMap.GetObjectSelection(EditorSubobject).GetPropertySelection())
 				{
 					// Recursively check whether subobjects also have a registered ICustomObjectSnapshotSerializer
 					const FRestoreObjectScope FinishRestore = PreObjectRestore_EditorWorld(SnapshotSubobject, EditorSubobject, WorldData, SelectionMap, LocalisationSnapshotPackage, [&WorldData, OriginalPath = MetaData->GetOriginalPath()](){ return WorldData.GetCustomSubobjectData_ForSubobject(OriginalPath);} );
@@ -106,6 +106,18 @@ namespace
 					CustomSerializer->OnPostSerializeEditorSubobject(EditorSubobject, *MetaData, SerializationDataReader);
 					continue;
 				}
+
+				const FCustomSubobjectRestorationInfo* RestorationInfo = SelectionMap.GetObjectSelection(EditorObject).GetCustomSubobjectSelection();
+				if (RestorationInfo && RestorationInfo->CustomSnapshotSubobjectsToRestore.Contains(SnapshotSubobject))
+				{
+					// Recursively check whether subobjects also have a registered ICustomObjectSnapshotSerializer
+					const FRestoreObjectScope FinishRestore = PreObjectRestore_EditorWorld(SnapshotSubobject, EditorSubobject, WorldData, SelectionMap, LocalisationSnapshotPackage, [&WorldData, OriginalPath = MetaData->GetOriginalPath()](){ return WorldData.GetCustomSubobjectData_ForSubobject(OriginalPath);} );
+
+					FCustomSerializationData* SerializationData = SerializationDataGetter();
+					FApplySnapshotDataArchiveV2::ApplyToRecreatedEditorWorldObject(SerializationData->Subobjects[i], WorldData, EditorSubobject, SnapshotSubobject, SelectionMap);
+					CustomSerializer->OnPostSerializeEditorSubobject(EditorSubobject, *MetaData, SerializationDataReader);
+					continue;
+				}
 				
 				UE_LOG(LogLevelSnapshots, Warning, TEXT("Editor subobject %s was not restored"), *EditorSubobject->GetPathName());	
 			}
@@ -240,7 +252,7 @@ FRestoreObjectScope FCustomObjectSerializationWrapper::PreSubobjectRestore_Edito
 		);
 }
 
-void FCustomObjectSerializationWrapper::ForEachMatchingCustomSubobjectPair(const FWorldSnapshotData& WorldData, UObject* SnapshotObject, UObject* WorldObject, FHandleCustomSubobjectPair Callback)
+void FCustomObjectSerializationWrapper::ForEachMatchingCustomSubobjectPair(const FWorldSnapshotData& WorldData, UObject* SnapshotObject, UObject* WorldObject, FHandleCustomSubobjectPair HandleCustomSubobjectPair,  FHandleUnmatchedCustomSnapshotSubobject HandleUnmachtedCustomSnapshotSubobject)
 {
 	FLevelSnapshotsModule& LevelSnapshots = FModuleManager::Get().GetModuleChecked<FLevelSnapshotsModule>("LevelSnapshots");
     TSharedPtr<ICustomObjectSnapshotSerializer> CustomSerializer = LevelSnapshots.GetCustomSerializerForClass(WorldObject->GetClass());
@@ -261,11 +273,19 @@ void FCustomObjectSerializationWrapper::ForEachMatchingCustomSubobjectPair(const
     	const TSharedPtr<ISnapshotSubobjectMetaData> MetaData = SerializationDataReader.GetSubobjectMetaData(i);
         UObject* EditorSubobject = CustomSerializer->FindSubobjectInEditorWorld(WorldObject, *MetaData, SerializationDataReader);
         UObject* SnapshotCounterpart = CustomSerializer->FindOrRecreateSubobjectInSnapshotWorld(SnapshotObject, *MetaData, SerializationDataReader);
-    	if (!EditorSubobject || !SnapshotCounterpart || !ensure(EditorSubobject->IsIn(WorldObject) && SnapshotCounterpart->IsIn(SnapshotObject)))
+    	if (!SnapshotCounterpart || !ensure(SnapshotCounterpart->IsIn(SnapshotObject)))
     	{
     		UE_LOG(LogLevelSnapshots, Warning, TEXT("FindOrRecreateSubobjectInSnapshotWorld did not return any valid subobject. Skipping subobject restoration..."));
     		continue;
     	}
-    	Callback(SnapshotCounterpart, EditorSubobject);
+
+    	if (!EditorSubobject)
+    	{
+    		HandleUnmachtedCustomSnapshotSubobject(SnapshotCounterpart);
+    	}
+    	else if (ensure(EditorSubobject->IsIn(WorldObject)))
+    	{
+    		HandleCustomSubobjectPair(SnapshotCounterpart, EditorSubobject);
+    	}
     }
 }
diff --git a/Engine/Plugins/VirtualProduction/LevelSnapshots/Source/LevelSnapshots/Private/Data/CustomSerialization/CustomObjectSerializationWrapper.h b/Engine/Plugins/VirtualProduction/LevelSnapshots/Source/LevelSnapshots/Private/Data/CustomSerialization/CustomObjectSerializationWrapper.h
--- a/Engine/Plugins/VirtualProduction/LevelSnapshots/Source/LevelSnapshots/Private/Data/CustomSerialization/CustomObjectSerializationWrapper.h
+++ b/Engine/Plugins/VirtualProduction/LevelSnapshots/Source/LevelSnapshots/Private/Data/CustomSerialization/CustomObjectSerializationWrapper.h
@@ -52,8 +52,6 @@ class FCustomObjectSerializationWrapper
 {
 public:
 	
-	using FHandleCustomSubobjectPair = TFunction<void(UObject* SnapshotSubobject, UObject* EditorSubobject)>;
-	
 	static void TakeSnapshotForActor(
 		AActor* EditorActor,
 		FCustomSerializationData& ActorSerializationData,
@@ -98,6 +96,10 @@ public:
 		UPackage* LocalisationSnapshotPackage
 		);
 
-	static void ForEachMatchingCustomSubobjectPair(const FWorldSnapshotData& WorldData, UObject* SnapshotObject, UObject* WorldObject, FHandleCustomSubobjectPair Callback);
+	
+	using FHandleCustomSubobjectPair = TFunction<void(UObject* SnapshotSubobject, UObject* EditorSubobject)>;
+	using FHandleUnmatchedCustomSnapshotSubobject = TFunction<void(UObject* UnmatchedSnapshotSubobject)>;
+	
+	static void ForEachMatchingCustomSubobjectPair(const FWorldSnapshotData& WorldData, UObject* SnapshotObject, UObject* WorldObject, FHandleCustomSubobjectPair HandleCustomSubobjectPair, FHandleUnmatchedCustomSnapshotSubobject HandleUnmachtedCustomSnapshotSubobject);
 };
 	
diff --git a/Engine/Plugins/VirtualProduction/LevelSnapshots/Source/LevelSnapshots/Private/Data/CustomSerialization/CustomSerializationDataManager.cpp b/Engine/Plugins/VirtualProduction/LevelSnapshots/Source/LevelSnapshots/Private/Data/CustomSerialization/CustomSerializationDataManager.cpp
--- a/Engine/Plugins/VirtualProduction/LevelSnapshots/Source/LevelSnapshots/Private/Data/CustomSerialization/CustomSerializationDataManager.cpp
+++ b/Engine/Plugins/VirtualProduction/LevelSnapshots/Source/LevelSnapshots/Private/Data/CustomSerialization/CustomSerializationDataManager.cpp
@@ -7,6 +7,9 @@
 #include "TakeWorldObjectSnapshotArchive.h"
 #include "Serialization/MemoryWriter.h"
 #include "Serialization/MemoryReader.h"
+#include "Serialization/NameAsStringProxyArchive.h"
+#include "Serialization/ObjectAndNameAsStringProxyArchive.h"
+#include "Util/SnapshotObjectUtil.h"
 
 class FSnapshotSubobjectMetaDataManager
 	:
@@ -46,18 +49,19 @@ public:
 		if (Writer.IsBound() && ensure(SerializationDataGetter_ReadWrite.IsBound()))
 		{
 			FMemoryWriter MemoryWriter(SerializationDataGetter_ReadWrite.Execute()->Subobjects[OwningDataSubobjectIndex].SubobjectAnnotationData, true);
-			
-			Writer.Execute(MemoryWriter);
+			FObjectAndNameAsStringProxyArchive RootArchive(MemoryWriter, false);
+			Writer.Execute(RootArchive);
 		}
 	}
 	
 	virtual void ReadObjectAnnotation(const FObjectAnnotator& Reader) const override
 	{
-		if (Reader.IsBound())
+		if (Reader.IsBound() && ensure(SerializationDataGetter_ReadOnly.IsBound()))
 		{
 			FMemoryReader MemoryReader(SerializationDataGetter_ReadOnly.Execute()->Subobjects[OwningDataSubobjectIndex].SubobjectAnnotationData, true);
-			WorldData.GetSnapshotVersionInfo().ApplyToArchive(MemoryReader);
-			Reader.Execute(MemoryReader);
+			FObjectAndNameAsStringProxyArchive RootArchive(MemoryReader, true);
+			WorldData.GetSnapshotVersionInfo().ApplyToArchive(RootArchive);
+			Reader.Execute(RootArchive);
 		}
 	}
 };
@@ -96,8 +100,9 @@ void FCustomSerializationDataReader::ReadObjectAnnotation(const FObjectAnnotator
 	if (Reader.IsBound())
 	{
 		FMemoryReader MemoryReader(SerializationDataGetter_ReadOnly.Execute()->RootAnnotationData, true);
-		WorldData_ReadOnly.GetSnapshotVersionInfo().ApplyToArchive(MemoryReader);
-		Reader.Execute(MemoryReader);
+		FObjectAndNameAsStringProxyArchive RootArchive(MemoryReader, true);
+		WorldData_ReadOnly.GetSnapshotVersionInfo().ApplyToArchive(RootArchive);
+		Reader.Execute(RootArchive);
 	}
 }
 
@@ -151,7 +156,8 @@ void FCustomSerializationDataWriter::WriteObjectAnnotation(const FObjectAnnotato
 	if (Writer.IsBound())
 	{
 		FMemoryWriter MemoryWriter(SerializationDataGetter_ReadWrite.Execute()->RootAnnotationData, true);
-		Writer.Execute(MemoryWriter);
+		FObjectAndNameAsStringProxyArchive RootArchive(MemoryWriter, false);
+		Writer.Execute(RootArchive);
 	}
 }
 
@@ -184,7 +190,7 @@ int32 FCustomSerializationDataWriter::AddSubobjectSnapshot(UObject* Subobject)
 	}
 
 	FCustomSubbjectSerializationData SubobjectData;
-	SubobjectData.ObjectPathIndex = WorldData_ReadWrite.AddObjectDependency(Subobject);
+	SubobjectData.ObjectPathIndex = SnapshotUtil::Object::AddObjectDependency(WorldData_ReadWrite, Subobject, false);
 	const int32 SubobjectIndex = SerializationData->Subobjects.Emplace(
 		MoveTemp(SubobjectData) // Not profiled
 		);
