Refresh UE5 branch with latest Level Snapshots which includes support for subobjects.

diff --git a/Engine/Plugins/VirtualProduction/LevelSnapshots/Source/LevelSnapshots/Private/Data/LevelSnapshot.cpp b/Engine/Plugins/VirtualProduction/LevelSnapshots/Source/LevelSnapshots/Private/Data/LevelSnapshot.cpp
--- a/Engine/Plugins/VirtualProduction/LevelSnapshots/Source/LevelSnapshots/Private/Data/LevelSnapshot.cpp
+++ b/Engine/Plugins/VirtualProduction/LevelSnapshots/Source/LevelSnapshots/Private/Data/LevelSnapshot.cpp
@@ -4,8 +4,7 @@
 
 #include "LevelSnapshotsLog.h"
 #include "LevelSnapshotsModule.h"
-#include "PropertyInfoHelpers.h"
-#include "Restorability/PropertyComparisonParams.h"
+#include "Util/EquivalenceUtil.h"
 
 #include "Algo/Accumulate.h"
 #include "Engine/Engine.h"
@@ -13,11 +12,9 @@
 #include "Engine/Engine.h"
 #include "GameFramework/Actor.h"
 #include "SnapshotRestorability.h"
-#include "CustomSerialization/CustomObjectSerializationWrapper.h"
 #include "CustomSerialization/CustomSerializationDataManager.h"
 #include "Misc/ScopeExit.h"
 #include "UObject/Package.h"
-#include "UObject/TextProperty.h"
 #include "UObject/UnrealType.h"
 #if WITH_EDITOR && !(UE_BUILD_SHIPPING || UE_BUILD_TEST)
 #include "Logging/MessageLog.h"
@@ -26,73 +23,6 @@
 #include "ScopedTransaction.h"
 #endif
 
-namespace
-{
-	/* If this function return false, the objects are not equivalent. If true, ignore the object references. */
-	bool ShouldConsiderObjectsEquivalent(const FWorldSnapshotData& SnapshotData, const FObjectPropertyBase* ObjectProperty, void* SnapshotValuePtr, void* WorldValuePtr, AActor* SnapshotActor, AActor* WorldActor)
-	{
-		// Native identity check handles:
-			// - external references, e.g. UMaterial in content browser
-			// - soft object paths: if SnapshotValuePtr is a TSoftObjectPtr<AActor> property, it points to a world actor instead of to an equivalent snapshot actor
-		if (ObjectProperty->Identical(SnapshotValuePtr, WorldValuePtr, 0))
-		{
-			return true;
-		}
-		
-		UObject* SnapshotObject = ObjectProperty->GetObjectPropertyValue(SnapshotValuePtr);
-		UObject* WorldObject = ObjectProperty->GetObjectPropertyValue(WorldValuePtr);
-		return SnapshotData.AreReferencesEquivalent(SnapshotObject, WorldObject, SnapshotActor, WorldActor);
-	}
-
-	void EnqueueMatchingCustomSubobjects(TInlineComponentArray<TPair<UObject*, UObject*>>& SnapshotOriginalPairsToProcess, const FWorldSnapshotData& WorldData, UObject* SnapshotObject, UObject* WorldObject)
-	{
-		FCustomObjectSerializationWrapper::ForEachMatchingCustomSubobjectPair(WorldData, SnapshotObject, WorldObject, [&SnapshotOriginalPairsToProcess, &WorldData](UObject* SnapshotSubobject, UObject* WorldSubobject)
-		{
-			SnapshotOriginalPairsToProcess.Add(TPair<UObject*, UObject*>(SnapshotSubobject, WorldSubobject));
-			EnqueueMatchingCustomSubobjects(SnapshotOriginalPairsToProcess, WorldData, SnapshotSubobject, WorldSubobject);
-		});
-	}
-	
-	void EnqueueMatchingComponents(TInlineComponentArray<TPair<UObject*, UObject*>>& SnapshotOriginalPairsToProcess, const FWorldSnapshotData& WorldData, AActor* SnapshotActor, AActor* WorldActor)
-	{
-		TInlineComponentArray<UActorComponent*> SnapshotComponents;
-		TInlineComponentArray<UActorComponent*> WorldComponents;
-		SnapshotActor->GetComponents(SnapshotComponents);
-		WorldActor->GetComponents(WorldComponents);
-		for (UActorComponent* WorldComponent : WorldComponents)
-		{
-			const bool bIsSupported = FSnapshotRestorability::IsComponentDesirableForCapture(WorldComponent);
-			if (!bIsSupported)
-			{
-				continue;
-			}
-
-			UActorComponent** SnapshotComponent = SnapshotComponents.FindByPredicate([WorldComponent](UActorComponent* Other)
-            {
-                return Other->GetFName() == WorldComponent->GetFName();
-            });
-			if (SnapshotComponent == nullptr)
-			{
-				UE_LOG(LogLevelSnapshots, Warning, TEXT("Failed to find component named '%s' of actor '%s' on snapshot version. Has the component name changed since the snapshot was taken?"), *WorldComponent->GetName(), *WorldActor->GetName());
-				continue;
-			}
-			UClass* SnapshotCompClass = (*SnapshotComponent)->GetClass();
-			UClass* WorldCompClass = WorldComponent->GetClass();
-			const bool bHaveSameClasses = SnapshotCompClass == WorldCompClass;
-			if (!bHaveSameClasses)
-			{
-				UE_LOG(LogLevelSnapshots, Warning,
-                    TEXT("World component '%s' of actor '%s' has class '%s' but snapshot version had class '%s'. Has the component changed classes since the snapshot was taken?"),
-                    *WorldComponent->GetName(), *WorldActor->GetName(), *WorldCompClass->GetName(), *SnapshotCompClass->GetName());
-				continue;
-			}
-
-			EnqueueMatchingCustomSubobjects(SnapshotOriginalPairsToProcess, WorldData, *SnapshotComponent, WorldComponent);
-			SnapshotOriginalPairsToProcess.Add(TPair<UObject*, UObject*>(*SnapshotComponent, WorldComponent));
-		}
-	}
-}
-
 void ULevelSnapshot::ApplySnapshotToWorld(UWorld* TargetWorld, const FPropertySelectionMap& SelectionSet)
 {
 	SCOPED_SNAPSHOT_CORE_TRACE(ApplyToWorld);
@@ -175,149 +105,16 @@ bool ULevelSnapshot::SnapshotWorld(UWorld* TargetWorld)
 
 bool ULevelSnapshot::HasOriginalChangedPropertiesSinceSnapshotWasTaken(AActor* SnapshotActor, AActor* WorldActor) const
 {
-	SCOPED_SNAPSHOT_CORE_TRACE(HasOriginalChangedSinceSnapshot);
-
-	if (!IsValid(SnapshotActor) || !IsValid(WorldActor))
-	{
-		return SnapshotActor != WorldActor;
-	}
-		
-	UClass* SnapshotClass = SnapshotActor->GetClass();
-	UClass* WorldClass = WorldActor->GetClass();
-	if (SnapshotClass != WorldClass)
-	{
-		return true;
-	}
-
-	TInlineComponentArray<TPair<UObject*, UObject*>> SnapshotOriginalPairsToProcess;
-	SnapshotOriginalPairsToProcess.Add(TPair<UObject*, UObject*>(SnapshotActor, WorldActor));
-	EnqueueMatchingCustomSubobjects(SnapshotOriginalPairsToProcess, SerializedData,SnapshotActor, WorldActor);
-	EnqueueMatchingComponents(SnapshotOriginalPairsToProcess, SerializedData,SnapshotActor, WorldActor);
-
-	FLevelSnapshotsModule& Module = FModuleManager::Get().GetModuleChecked<FLevelSnapshotsModule>("LevelSnapshots");
-	for (const TPair<UObject*, UObject*>& NextPair : SnapshotOriginalPairsToProcess)
-	{
-		UClass* ClassToIterate = NextPair.Key->GetClass();
-		UObject* const SnapshotObject = NextPair.Key;
-		UObject* const WorldObject = NextPair.Value;
-		
-		const FPropertyComparerArray PropertyComparers = Module.GetPropertyComparerForClass(ClassToIterate);
-		for (TFieldIterator<FProperty> FieldIt(ClassToIterate); FieldIt; ++FieldIt)
-		{
-			// Ask external modules about the property
-			const FPropertyComparisonParams Params { WorldClass, *FieldIt, SnapshotObject, WorldObject, SnapshotObject, WorldObject, SnapshotActor, WorldActor} ;
-			const IPropertyComparer::EPropertyComparison ComparisonResult = Module.ShouldConsiderPropertyEqual(PropertyComparers, Params);
-			
-			switch (ComparisonResult)
-			{
-			case IPropertyComparer::EPropertyComparison::TreatEqual:
-				continue;
-			case IPropertyComparer::EPropertyComparison::TreatUnequal:
-				return true;
-			default:
-				if (!AreSnapshotAndOriginalPropertiesEquivalent(*FieldIt, SnapshotObject, WorldObject, SnapshotActor, WorldActor))
-				{
-					return true;
-				}
-			}
-		}
-	}
-	return false;
+	return SnapshotUtil::HasOriginalChangedPropertiesSinceSnapshotWasTaken(SerializedData, SnapshotActor, WorldActor);
 }
 
 bool ULevelSnapshot::AreSnapshotAndOriginalPropertiesEquivalent(const FProperty* LeafProperty, void* SnapshotContainer, void* WorldContainer, AActor* SnapshotActor, AActor* WorldActor) const
 {
-	// Ensure that property's flags are allowed. Skip check collection properties, e.g. FArrayProperty::Inner, etc.: inner properties do not have the same flags.
-	const bool bIsInCollection = FPropertyInfoHelpers::IsPropertyInCollection(LeafProperty);
-	if (!bIsInCollection && !FSnapshotRestorability::IsRestorableProperty(LeafProperty))
-	{
-		return true;
-	}
-	
-	for (int32 i = 0; i < LeafProperty->ArrayDim; ++i)
-	{
-		void* SnapshotValuePtr = LeafProperty->ContainerPtrToValuePtr<void>(SnapshotContainer, i);
-		void* WorldValuePtr = LeafProperty->ContainerPtrToValuePtr<void>(WorldContainer, i);
-
-		// Check whether float is nearly equal instead of exactly equal
-		if (const FNumericProperty* NumericProperty = CastField<FNumericProperty>(LeafProperty))
-		{
-			return FPropertyInfoHelpers::AreNumericPropertiesNearlyEqual(NumericProperty, SnapshotValuePtr, WorldValuePtr);
-		}
-		
-		// Use our custom equality function for struct properties
-		if (const FStructProperty* StructProperty = CastField<FStructProperty>(LeafProperty))
-		{
-			for (TFieldIterator<FProperty> FieldIt(StructProperty->Struct); FieldIt; ++FieldIt)
-			{
-				if (!AreSnapshotAndOriginalPropertiesEquivalent(*FieldIt, SnapshotValuePtr, WorldValuePtr, SnapshotActor, WorldActor))
-				{
-					return false;
-				}
-			}
-			return true;
-		}
-
-		// Check whether property value points to a subobject
-		if (const FObjectPropertyBase* ObjectProperty = CastField<FObjectPropertyBase>(LeafProperty))
-		{
-			return ShouldConsiderObjectsEquivalent(SerializedData, ObjectProperty, SnapshotValuePtr, WorldValuePtr, SnapshotActor, WorldActor);
-		}
-
-
-		// Use our custom equality function for array, set, and map inner properties
-		if (const FArrayProperty* ArrayProperty = CastField<FArrayProperty>(LeafProperty))
-		{
-			FScriptArrayHelper SnapshotArray(ArrayProperty, SnapshotValuePtr);
-			FScriptArrayHelper WorldArray(ArrayProperty, WorldValuePtr);
-			if (SnapshotArray.Num() != WorldArray.Num())
-			{
-				return false;
-			}
-
-			for (int32 j = 0; j < SnapshotArray.Num(); ++j)
-			{
-				void* SnapshotElementValuePtr = SnapshotArray.GetRawPtr(j);
-				void* WorldElementValuePtr = WorldArray.GetRawPtr(j);
-
-				if (!AreSnapshotAndOriginalPropertiesEquivalent(ArrayProperty->Inner, SnapshotElementValuePtr, WorldElementValuePtr, SnapshotActor, WorldActor))
-				{
-					return false;
-				}
-			}
-			return true;
-		}
-		
-		if (const FMapProperty* MapProperty = CastField<FMapProperty>(LeafProperty))
-		{
-			// TODO: Use custom function. Need to do something similar to UE4MapProperty_Private::IsPermutation
-		}
-
-		if (const FSetProperty* SetProperty = CastField<FSetProperty>(LeafProperty))
-		{
-			// TODO: Use custom function. Need to do something similar to UE4SetProperty_Private::IsPermutation
-		}
-
-		if (const FTextProperty* TextProperty = CastField<FTextProperty>(LeafProperty))
-		{
-			const FText& SnapshotText = TextProperty->GetPropertyValue_InContainer(SnapshotContainer);
-			const FText& WorldText = TextProperty->GetPropertyValue_InContainer(WorldContainer);
-			return SnapshotText.IdenticalTo(WorldText, ETextIdenticalModeFlags::None) || SnapshotText.ToString().Equals(WorldText.ToString());
-		}
-		
-		// Use normal property comparison for all other properties
-		if (!LeafProperty->Identical_InContainer(SnapshotContainer, WorldContainer, i, PPF_DeepComparison | PPF_DeepCompareDSOsOnly))
-		{
-			return false;
-		}
-	}
-	return true;
+	return SnapshotUtil::AreSnapshotAndOriginalPropertiesEquivalent(SerializedData, LeafProperty, SnapshotContainer, WorldContainer, SnapshotActor, WorldActor);
 }
 
 TOptional<AActor*> ULevelSnapshot::GetDeserializedActor(const FSoftObjectPath& OriginalActorPath)
 {
-	SCOPED_SNAPSHOT_CORE_TRACE(GetDeserializedActor);
-	
 	EnsureWorldInitialised();
 	return SerializedData.GetDeserializedActor(OriginalActorPath, GetPackage());
 }
@@ -369,7 +166,6 @@ void ULevelSnapshot::DiffWorld(UWorld* World, FActorPathConsumer HandleMatchedAc
 	// Try to find world actors and call appropriate callback
 	{
 		SCOPED_SNAPSHOT_CORE_TRACE(DiffWorld_IteratorAllActors);
-		
 		SerializedData.ForEachOriginalActor([&HandleMatchedActor, &HandleRemovedActor, &HandleAddedActor, &AllActors](const FSoftObjectPath& OriginalActorPath, const FActorSnapshotData& SavedData)
 		{
 			// TODO: we need to check whether the actor's class was blacklisted in the project settings
diff --git a/Engine/Plugins/VirtualProduction/LevelSnapshots/Source/LevelSnapshots/Private/Data/PropertySelection.cpp b/Engine/Plugins/VirtualProduction/LevelSnapshots/Source/LevelSnapshots/Private/Data/PropertySelection.cpp
--- a/Engine/Plugins/VirtualProduction/LevelSnapshots/Source/LevelSnapshots/Private/Data/PropertySelection.cpp
+++ b/Engine/Plugins/VirtualProduction/LevelSnapshots/Source/LevelSnapshots/Private/Data/PropertySelection.cpp
@@ -22,6 +22,11 @@ namespace
 	}
 }
 
+FLevelSnapshotPropertyChain::FLevelSnapshotPropertyChain(const FProperty* RootProperty)
+{
+	AppendInline(RootProperty);
+}
+
 FLevelSnapshotPropertyChain FLevelSnapshotPropertyChain::MakeAppended(const FProperty* Property) const
 {
 	FLevelSnapshotPropertyChain Result = *this;
@@ -95,6 +100,19 @@ bool FLevelSnapshotPropertyChain::operator==(const FLevelSnapshotPropertyChain&
 	return true;
 }
 
+FPropertySelection::FPropertySelection(const FProperty* SingleProperty)
+{
+	AddProperty({ SingleProperty });
+}
+
+FPropertySelection::FPropertySelection(const TSet<const FProperty*>& Properties)
+{
+	for (const FProperty* Property : Properties)
+	{
+		AddProperty({ Property  });
+	}
+}
+
 bool FPropertySelection::ShouldSerializeProperty(const FArchiveSerializedPropertyChain* ContainerChain, const FProperty* LeafProperty) const
 {
 	if (IsPropertySelected(ContainerChain, LeafProperty))
diff --git a/Engine/Plugins/VirtualProduction/LevelSnapshots/Source/LevelSnapshots/Private/Data/RestorableObjectSelection.cpp b/Engine/Plugins/VirtualProduction/LevelSnapshots/Source/LevelSnapshots/Private/Data/RestorableObjectSelection.cpp
new file mode 100644
--- /dev/null
+++ b/Engine/Plugins/VirtualProduction/LevelSnapshots/Source/LevelSnapshots/Private/Data/RestorableObjectSelection.cpp
@@ -0,0 +1,22 @@
+﻿// Copyright Epic Games, Inc. All Rights Reserved.
+
+#include "Data/RestorableObjectSelection.h"
+
+#include "AddedAndRemovedComponentInfo.h"
+#include "PropertySelectionMap.h"
+#include "PropertySelection.h"
+
+const FPropertySelection* FRestorableObjectSelection::GetPropertySelection() const
+{
+	return Owner.EditorWorldObjectToSelectedProperties.Find(ObjectPath);
+}
+
+const FAddedAndRemovedComponentInfo* FRestorableObjectSelection::GetComponentSelection() const
+{
+	return Owner.EditorActorToComponentSelection.Find(ObjectPath);
+}
+
+const FCustomSubobjectRestorationInfo* FRestorableObjectSelection::GetCustomSubobjectSelection() const
+{
+	return Owner.EditorWorldObjectToCustomSubobjectSelection.Find(ObjectPath);
+}
diff --git a/Engine/Plugins/VirtualProduction/LevelSnapshots/Source/LevelSnapshots/Private/Data/RestorationEvents/ApplySnapshotPropertiesScope.cpp b/Engine/Plugins/VirtualProduction/LevelSnapshots/Source/LevelSnapshots/Private/Data/RestorationEvents/ApplySnapshotPropertiesScope.cpp
new file mode 100644
--- /dev/null
+++ b/Engine/Plugins/VirtualProduction/LevelSnapshots/Source/LevelSnapshots/Private/Data/RestorationEvents/ApplySnapshotPropertiesScope.cpp
@@ -0,0 +1,17 @@
+﻿// Copyright Epic Games, Inc. All Rights Reserved.
+
+#include "Data/RestorationEvents/ApplySnapshotPropertiesScope.h"
+
+#include "LevelSnapshotsModule.h"
+
+FApplySnapshotPropertiesScope::FApplySnapshotPropertiesScope(FApplySnapshotPropertiesParams Params)
+	:
+	Params(Params)
+{
+	FLevelSnapshotsModule::GetInternalModuleInstance().OnPreApplySnapshotProperties(Params);
+}
+
+FApplySnapshotPropertiesScope::~FApplySnapshotPropertiesScope()
+{
+	FLevelSnapshotsModule::GetInternalModuleInstance().OnPostApplySnapshotProperties(Params);
+}
diff --git a/Engine/Plugins/VirtualProduction/LevelSnapshots/Source/LevelSnapshots/Private/Data/RestorationEvents/ApplySnapshotPropertiesScope.h b/Engine/Plugins/VirtualProduction/LevelSnapshots/Source/LevelSnapshots/Private/Data/RestorationEvents/ApplySnapshotPropertiesScope.h
new file mode 100644
--- /dev/null
+++ b/Engine/Plugins/VirtualProduction/LevelSnapshots/Source/LevelSnapshots/Private/Data/RestorationEvents/ApplySnapshotPropertiesScope.h
@@ -0,0 +1,21 @@
+﻿// Copyright Epic Games, Inc. All Rights Reserved.
+
+#pragma once
+
+#include "CoreMinimal.h"
+#include "Restorability/IRestorationListener.h"
+#include "Templates/UnrealTemplate.h"
+
+struct FPropertySelectionMap;
+
+/**
+ * Convenience type that calls FLevelSnaphshotsModule::OnPreApplySnapshot and FLevelSnaphshotsModule::OnPostApplySnapshot.
+ */
+class FApplySnapshotPropertiesScope : public FNoncopyable
+{
+	const FApplySnapshotPropertiesParams Params;
+public:
+
+	FApplySnapshotPropertiesScope(FApplySnapshotPropertiesParams Params);
+	~FApplySnapshotPropertiesScope();
+};
\ No newline at end of file
diff --git a/Engine/Plugins/VirtualProduction/LevelSnapshots/Source/LevelSnapshots/Private/Data/RestorationEvents/ApplySnapshotToActorScope.cpp b/Engine/Plugins/VirtualProduction/LevelSnapshots/Source/LevelSnapshots/Private/Data/RestorationEvents/ApplySnapshotToActorScope.cpp
new file mode 100644
--- /dev/null
+++ b/Engine/Plugins/VirtualProduction/LevelSnapshots/Source/LevelSnapshots/Private/Data/RestorationEvents/ApplySnapshotToActorScope.cpp
@@ -0,0 +1,17 @@
+﻿// Copyright Epic Games, Inc. All Rights Reserved.
+
+#include "Data/RestorationEvents/ApplySnapshotToActorScope.h"
+
+#include "LevelSnapshotsModule.h"
+
+FApplySnapshotToActorScope::FApplySnapshotToActorScope(FApplySnapshotToActorParams Params)
+	:
+	Params(Params)
+{
+	FLevelSnapshotsModule::GetInternalModuleInstance().OnPreApplySnapshotToActor(Params);
+}
+
+FApplySnapshotToActorScope::~FApplySnapshotToActorScope()
+{
+	FLevelSnapshotsModule::GetInternalModuleInstance().OnPostApplySnapshotToActor(Params);
+}
diff --git a/Engine/Plugins/VirtualProduction/LevelSnapshots/Source/LevelSnapshots/Private/Data/RestorationEvents/ApplySnapshotToActorScope.h b/Engine/Plugins/VirtualProduction/LevelSnapshots/Source/LevelSnapshots/Private/Data/RestorationEvents/ApplySnapshotToActorScope.h
new file mode 100644
--- /dev/null
+++ b/Engine/Plugins/VirtualProduction/LevelSnapshots/Source/LevelSnapshots/Private/Data/RestorationEvents/ApplySnapshotToActorScope.h
@@ -0,0 +1,21 @@
+﻿// Copyright Epic Games, Inc. All Rights Reserved.
+
+#pragma once
+
+#include "CoreMinimal.h"
+#include "Restorability/IRestorationListener.h"
+#include "Templates/UnrealTemplate.h"
+
+struct FPropertySelectionMap;
+
+/**
+ * Convenience type that calls FLevelSnaphshotsModule::OnPreApplySnapshot and FLevelSnaphshotsModule::OnPostApplySnapshot.
+ */
+class FApplySnapshotToActorScope : public FNoncopyable
+{
+	const FApplySnapshotToActorParams Params;
+public:
+
+	FApplySnapshotToActorScope(FApplySnapshotToActorParams Params);
+	~FApplySnapshotToActorScope();
+};
\ No newline at end of file
diff --git a/Engine/Plugins/VirtualProduction/LevelSnapshots/Source/LevelSnapshots/Private/Data/RestorationEvents/RecreateComponentScope.cpp b/Engine/Plugins/VirtualProduction/LevelSnapshots/Source/LevelSnapshots/Private/Data/RestorationEvents/RecreateComponentScope.cpp
new file mode 100644
--- /dev/null
+++ b/Engine/Plugins/VirtualProduction/LevelSnapshots/Source/LevelSnapshots/Private/Data/RestorationEvents/RecreateComponentScope.cpp
@@ -0,0 +1,32 @@
+﻿// Copyright Epic Games, Inc. All Rights Reserved.
+
+#include "Data/RestorationEvents/RecreateComponentScope.h"
+
+#include "LevelSnapshotsLog.h"
+#include "LevelSnapshotsModule.h"
+#include "SubobjectSnapshotData.h"
+
+FRecreateComponentScope::FRecreateComponentScope(
+	const FSubobjectSnapshotData& SubobjectSnapshotData,
+	AActor* Owner,
+	FName ComponentName,
+	UClass* ComponentClass,
+	EComponentCreationMethod CreationMethod)
+		:
+	SubobjectSnapshotData(SubobjectSnapshotData)
+{
+	FLevelSnapshotsModule::GetInternalModuleInstance().OnPreRecreateComponent({ Owner, ComponentName, ComponentClass, CreationMethod });
+}
+
+FRecreateComponentScope::~FRecreateComponentScope()
+{
+	// Creation might fail in exceptional cases
+	if (UActorComponent* Component = Cast<UActorComponent>(SubobjectSnapshotData.EditorObject))
+	{
+		FLevelSnapshotsModule::GetInternalModuleInstance().OnPostRecreateComponent(Component);
+	}
+	else
+	{
+		UE_LOG(LogLevelSnapshots, Warning, TEXT("Skipping calls to IRestorationListener::PostRecreateComponent"));
+	}
+}
diff --git a/Engine/Plugins/VirtualProduction/LevelSnapshots/Source/LevelSnapshots/Private/Data/RestorationEvents/RecreateComponentScope.h b/Engine/Plugins/VirtualProduction/LevelSnapshots/Source/LevelSnapshots/Private/Data/RestorationEvents/RecreateComponentScope.h
new file mode 100644
--- /dev/null
+++ b/Engine/Plugins/VirtualProduction/LevelSnapshots/Source/LevelSnapshots/Private/Data/RestorationEvents/RecreateComponentScope.h
@@ -0,0 +1,25 @@
+﻿// Copyright Epic Games, Inc. All Rights Reserved.
+
+#pragma once
+
+#include "CoreMinimal.h"
+
+struct FSubobjectSnapshotData;
+
+/**
+ * Convenience type that calls FLevelSnaphshotsModule::OnPreRecreateComponent and FLevelSnaphshotsModule::OnPostRecreateComponent.
+ */
+class FRecreateComponentScope : public FNoncopyable
+{
+	const FSubobjectSnapshotData& SubobjectSnapshotData;
+public:
+
+	FRecreateComponentScope(
+		const FSubobjectSnapshotData& SubobjectSnapshotData,
+		AActor* Owner,
+		FName ComponentName,
+		UClass* ComponentClass,
+		EComponentCreationMethod CreationMethod
+		);
+	~FRecreateComponentScope();
+};
\ No newline at end of file
diff --git a/Engine/Plugins/VirtualProduction/LevelSnapshots/Source/LevelSnapshots/Private/Data/RestorationEvents/RemoveComponentScope.cpp b/Engine/Plugins/VirtualProduction/LevelSnapshots/Source/LevelSnapshots/Private/Data/RestorationEvents/RemoveComponentScope.cpp
new file mode 100644
--- /dev/null
+++ b/Engine/Plugins/VirtualProduction/LevelSnapshots/Source/LevelSnapshots/Private/Data/RestorationEvents/RemoveComponentScope.cpp
@@ -0,0 +1,17 @@
+﻿// Copyright Epic Games, Inc. All Rights Reserved.
+
+#include "Data/RestorationEvents/RemoveComponentScope.h"
+
+#include "LevelSnapshotsModule.h"
+
+FRemoveComponentScope::FRemoveComponentScope(UActorComponent* RemovedComponent)
+	:
+	Params({ RemovedComponent->GetOwner(), RemovedComponent->GetFName(), RemovedComponent })
+{
+	FLevelSnapshotsModule::GetInternalModuleInstance().OnPreRemoveComponent(RemovedComponent);
+}
+
+FRemoveComponentScope::~FRemoveComponentScope()
+{
+	FLevelSnapshotsModule::GetInternalModuleInstance().OnPostRemoveComponent(Params);
+}
diff --git a/Engine/Plugins/VirtualProduction/LevelSnapshots/Source/LevelSnapshots/Private/Data/SnapshotCustomVersion.cpp b/Engine/Plugins/VirtualProduction/LevelSnapshots/Source/LevelSnapshots/Private/Data/SnapshotCustomVersion.cpp
new file mode 100644
--- /dev/null
+++ b/Engine/Plugins/VirtualProduction/LevelSnapshots/Source/LevelSnapshots/Private/Data/SnapshotCustomVersion.cpp
@@ -0,0 +1,7 @@
+// Copyright Epic Games, Inc. All Rights Reserved.
+
+#include "SnapshotCustomVersion.h"
+#include "Serialization/CustomVersion.h"
+
+const FGuid FSnapshotCustomVersion::GUID(0x3EF5FDBD, 0x21AAAF1F, 0xBCD6F38F, 0xB3A32521);
+FCustomVersionRegistration GRegisterSnapshotCustomVersion(FSnapshotCustomVersion::GUID, FSnapshotCustomVersion::LatestVersion, TEXT("LevelSnapshotVersion"));
diff --git a/Engine/Plugins/VirtualProduction/LevelSnapshots/Source/LevelSnapshots/Private/Data/SnapshotVersion.cpp b/Engine/Plugins/VirtualProduction/LevelSnapshots/Source/LevelSnapshots/Private/Data/SnapshotVersion.cpp
--- a/Engine/Plugins/VirtualProduction/LevelSnapshots/Source/LevelSnapshots/Private/Data/SnapshotVersion.cpp
+++ b/Engine/Plugins/VirtualProduction/LevelSnapshots/Source/LevelSnapshots/Private/Data/SnapshotVersion.cpp
@@ -1,6 +1,9 @@
 // Copyright Epic Games, Inc. All Rights Reserved.
 
 #include "Data/SnapshotVersion.h"
+
+#include "SnapshotCustomVersion.h"
+
 #include "Misc/EngineVersion.h"
 #include "Serialization/CustomVersion.h"
 
@@ -25,7 +28,7 @@ void FSnapshotCustomVersionInfo::Initialize(const FCustomVersion& InVersion)
 	Version = InVersion.Version;
 }
 
-void FSnapshotVersionInfo::Initialize()
+void FSnapshotVersionInfo::Initialize(bool bWithoutSnapshotVersion)
 {
 	FileVersion.Initialize();
 	EngineVersion.Initialize(FEngineVersion::Current());
@@ -35,7 +38,10 @@ void FSnapshotVersionInfo::Initialize()
 	for (const FCustomVersion& EngineCustomVersion : AllCurrentVersions.GetAllVersions())
 	{
 		FSnapshotCustomVersionInfo& CustomVersion = CustomVersions.AddDefaulted_GetRef();
-		CustomVersion.Initialize(EngineCustomVersion);
+		if (!bWithoutSnapshotVersion || (bWithoutSnapshotVersion && CustomVersion.Key != FSnapshotCustomVersion::GUID))
+		{
+			CustomVersion.Initialize(EngineCustomVersion);
+		}
 	}
 }
 
@@ -60,3 +66,16 @@ void FSnapshotVersionInfo::ApplyToArchive(FArchive& Archive) const
 		Archive.SetCustomVersions(EngineCustomVersions);
 	}
 }
+
+int32 FSnapshotVersionInfo::GetSnapshotCustomVersion() const
+{
+	for (const FSnapshotCustomVersionInfo& CustomVersion : CustomVersions)
+	{
+		if (CustomVersion.Key == FSnapshotCustomVersion::GUID)
+		{
+			return CustomVersion.Version;
+		}
+	}
+
+	return INDEX_NONE;
+}
diff --git a/Engine/Plugins/VirtualProduction/LevelSnapshots/Source/LevelSnapshots/Private/Data/Util/EquivalenceUtil.cpp b/Engine/Plugins/VirtualProduction/LevelSnapshots/Source/LevelSnapshots/Private/Data/Util/EquivalenceUtil.cpp
new file mode 100644
--- /dev/null
+++ b/Engine/Plugins/VirtualProduction/LevelSnapshots/Source/LevelSnapshots/Private/Data/Util/EquivalenceUtil.cpp
@@ -0,0 +1,398 @@
+﻿// Copyright Epic Games, Inc. All Rights Reserved.
+
+#include "EquivalenceUtil.h"
+
+#include "ApplySnapshotDataArchiveV2.h"
+#include "LevelSnapshotsLog.h"
+#include "WorldSnapshotData.h"
+
+#include "Components/ActorComponent.h"
+#include "EngineUtils.h"
+#include "LevelSnapshotsModule.h"
+#include "PropertyComparisonParams.h"
+#include "PropertyInfoHelpers.h"
+#include "SnapshotCustomVersion.h"
+#include "SnapshotRestorability.h"
+#include "CustomSerialization/CustomObjectSerializationWrapper.h"
+#include "GameFramework/Actor.h"
+#include "Misc/ScopeExit.h"
+
+namespace
+{
+	/** Iterates properties of SnapshotObject and WorldObject, owned by SnapshotActor and WorldActor, and returns whether at least one property value was different */
+	bool HaveDifferentPropertyValues(const FWorldSnapshotData& WorldData, UObject* SnapshotObject, UObject* WorldObject, AActor* SnapshotActor, AActor* WorldActor)
+	{
+		UClass* ClassToIterate = SnapshotObject->GetClass();
+		
+		FLevelSnapshotsModule& Module = FModuleManager::Get().GetModuleChecked<FLevelSnapshotsModule>("LevelSnapshots");
+		const FPropertyComparerArray PropertyComparers = Module.GetPropertyComparerForClass(ClassToIterate);
+		for (TFieldIterator<FProperty> FieldIt(ClassToIterate); FieldIt; ++FieldIt)
+		{
+			// Ask external modules about the property
+			const FPropertyComparisonParams Params { ClassToIterate, *FieldIt, SnapshotObject, WorldObject, SnapshotObject, WorldObject, SnapshotActor, WorldActor} ;
+			const IPropertyComparer::EPropertyComparison ComparisonResult = Module.ShouldConsiderPropertyEqual(PropertyComparers, Params);
+			
+			switch (ComparisonResult)
+			{
+			case IPropertyComparer::EPropertyComparison::TreatEqual:
+				continue;
+			case IPropertyComparer::EPropertyComparison::TreatUnequal:
+				return true;
+			default:
+				if (!SnapshotUtil::AreSnapshotAndOriginalPropertiesEquivalent(WorldData, *FieldIt, SnapshotObject, WorldObject, SnapshotActor, WorldActor))
+				{
+					return true;
+				}
+			}
+		}
+
+		return false;
+	}
+
+	bool EnqueueMatchingCustomSubobjects(TInlineComponentArray<TPair<UObject*, UObject*>>& SnapshotOriginalPairsToProcess, const FWorldSnapshotData& WorldData, UObject* SnapshotObject, UObject* WorldObject)
+	{
+		bool bFailedToMatchAllObjects = false;
+		FCustomObjectSerializationWrapper::ForEachMatchingCustomSubobjectPair(WorldData, SnapshotObject, WorldObject,
+			[&SnapshotOriginalPairsToProcess, &WorldData, &bFailedToMatchAllObjects](UObject* SnapshotSubobject, UObject* WorldSubobject)
+			{
+				SnapshotOriginalPairsToProcess.Add(TPair<UObject*, UObject*>(SnapshotSubobject, WorldSubobject));
+				if (!bFailedToMatchAllObjects)
+				{
+					bFailedToMatchAllObjects |= EnqueueMatchingCustomSubobjects(SnapshotOriginalPairsToProcess, WorldData, SnapshotSubobject, WorldSubobject);
+				}
+			},
+			[&bFailedToMatchAllObjects](UObject* UmatchedSnapshotSubobject)
+			{
+				bFailedToMatchAllObjects = true;
+			}
+		);
+
+		return bFailedToMatchAllObjects;
+	}
+
+	UActorComponent* TryFindMatchingComponent(UActorComponent* CounterpartComponentToMatch, const TInlineComponentArray<UActorComponent*>& ActorComponents)
+	{
+		const FName ComponentName = CounterpartComponentToMatch->GetFName();
+		UActorComponent* const* PossibleResult = ActorComponents.FindByPredicate([ComponentName](UActorComponent* Component)
+		{
+			return Component->GetFName().IsEqual(ComponentName);
+		});
+		if (!PossibleResult)
+		{
+			return nullptr;
+		}
+
+		UActorComponent* MatchedComponent = *PossibleResult;
+		if (MatchedComponent->GetClass() != CounterpartComponentToMatch->GetClass())
+		{
+			UE_LOG(LogLevelSnapshots, Error, TEXT("Components named %s were matched to each other but had differing classes (%s and %s)."),
+				*ComponentName.ToString(),
+				*MatchedComponent->GetClass()->GetName(),
+				*CounterpartComponentToMatch->GetClass()->GetName()
+			);
+			return nullptr;
+		}
+
+		return MatchedComponent;
+	}
+
+	/** @return False if one component could not be matched */
+	bool EnqueueMatchingComponents(TInlineComponentArray<TPair<UObject*, UObject*>>& SnapshotOriginalPairsToProcess, const FWorldSnapshotData& WorldData, AActor* SnapshotActor, AActor* WorldActor)
+	{
+		bool bFoundUnmatchedObjects = false;
+		SnapshotUtil::IterateComponents(SnapshotActor, WorldActor,
+			[&SnapshotOriginalPairsToProcess, &WorldData, &bFoundUnmatchedObjects](UActorComponent* SnapshotComp, UActorComponent* WorldComp)
+			{
+				bFoundUnmatchedObjects |= EnqueueMatchingCustomSubobjects(SnapshotOriginalPairsToProcess, WorldData, SnapshotComp, WorldComp);
+				SnapshotOriginalPairsToProcess.Add(TPair<UObject*, UObject*>(SnapshotComp, WorldComp));
+			},
+			[&bFoundUnmatchedObjects](UActorComponent* SnapshotComp)
+			{
+				bFoundUnmatchedObjects = true;
+			},
+			[&bFoundUnmatchedObjects](UActorComponent* WorldComp)
+			{
+				bFoundUnmatchedObjects = true;
+			});
+
+		return bFoundUnmatchedObjects;
+	}
+}
+
+void SnapshotUtil::IterateComponents(AActor* SnapshotActor, AActor* WorldActor, FHandleMatchedActorComponent OnComponentsMatched, FHandleUnmatchedActorComponent OnSnapshotComponentUnmatched, FHandleUnmatchedActorComponent OnWorldComponentUnmatched)
+{
+	TInlineComponentArray<UActorComponent*> WorldComponents(WorldActor);
+	TInlineComponentArray<UActorComponent*> SnapshotComponents(SnapshotActor);
+	
+	for (UActorComponent* WorldComp : WorldComponents)
+	{
+		if (!FSnapshotRestorability::IsComponentDesirableForCapture(WorldComp))
+		{
+			continue;
+		}
+		
+		if (UActorComponent* SnapshotMatchedComp = TryFindMatchingComponent(WorldComp, SnapshotComponents))
+		{
+			OnComponentsMatched(SnapshotMatchedComp, WorldComp);
+		}
+		else
+		{
+			OnWorldComponentUnmatched(WorldComp);
+		}
+	}
+
+	for (UActorComponent* SnapshotComp : SnapshotComponents)
+	{
+		if (FSnapshotRestorability::IsComponentDesirableForCapture(SnapshotComp)
+			&& TryFindMatchingComponent(SnapshotComp, WorldComponents) == nullptr)
+		{
+			OnSnapshotComponentUnmatched(SnapshotComp);
+		}
+	}
+}
+
+bool SnapshotUtil::HasOriginalChangedPropertiesSinceSnapshotWasTaken(const FWorldSnapshotData& WorldData, AActor* SnapshotActor, AActor* WorldActor)
+{
+	SCOPED_SNAPSHOT_CORE_TRACE(HasOriginalChangedSinceSnapshot);
+	
+	if (!IsValid(SnapshotActor) || !IsValid(WorldActor))
+	{
+		return SnapshotActor != WorldActor;
+	}
+		
+	UClass* SnapshotClass = SnapshotActor->GetClass();
+	UClass* WorldClass = WorldActor->GetClass();
+	if (SnapshotClass != WorldClass)
+	{
+		return true;
+	}
+
+	TInlineComponentArray<TPair<UObject*, UObject*>> SnapshotOriginalPairsToProcess;
+	SnapshotOriginalPairsToProcess.Add(TPair<UObject*, UObject*>(SnapshotActor, WorldActor));
+	
+	const bool bFailedToMatchAllComponentObjects = EnqueueMatchingComponents(SnapshotOriginalPairsToProcess, WorldData, SnapshotActor, WorldActor);
+	if (bFailedToMatchAllComponentObjects)
+	{
+		return true;
+	}
+	
+	const bool bFailedToMatchAllSubobjects = EnqueueMatchingCustomSubobjects(SnapshotOriginalPairsToProcess, WorldData, SnapshotActor, WorldActor);
+	if (bFailedToMatchAllSubobjects)
+	{
+		return true;
+	}
+	
+	for (const TPair<UObject*, UObject*>& NextPair : SnapshotOriginalPairsToProcess)
+	{
+		UObject* const SnapshotObject = NextPair.Key;
+		UObject* const WorldObject = NextPair.Value;
+		if (HaveDifferentPropertyValues(WorldData, SnapshotObject, WorldObject, SnapshotActor, WorldActor))
+		{
+			return true;
+		}
+	}
+	return false;
+}
+
+bool SnapshotUtil::AreSnapshotAndOriginalPropertiesEquivalent(const FWorldSnapshotData& WorldData, const FProperty* LeafProperty, void* SnapshotContainer, void* WorldContainer, AActor* SnapshotActor, AActor* WorldActor)
+{
+	// Ensure that property's flags are allowed. Skip check collection properties, e.g. FArrayProperty::Inner, etc.: inner properties do not have the same flags.
+	const bool bIsInnnerCollectionProperty = FPropertyInfoHelpers::IsPropertyInCollection(LeafProperty);
+	if (!bIsInnnerCollectionProperty && !FSnapshotRestorability::IsRestorableProperty(LeafProperty))
+	{
+		return true;
+	}
+	
+	for (int32 i = 0; i < LeafProperty->ArrayDim; ++i)
+	{
+		void* SnapshotValuePtr = LeafProperty->ContainerPtrToValuePtr<void>(SnapshotContainer, i);
+		void* WorldValuePtr = LeafProperty->ContainerPtrToValuePtr<void>(WorldContainer, i);
+
+		// Check whether float is nearly equal instead of exactly equal
+		if (const FNumericProperty* NumericProperty = CastField<FNumericProperty>(LeafProperty))
+		{
+			return FPropertyInfoHelpers::AreNumericPropertiesNearlyEqual(NumericProperty, SnapshotValuePtr, WorldValuePtr);
+		}
+		
+		// Use our custom equality function for struct properties
+		if (const FStructProperty* StructProperty = CastField<FStructProperty>(LeafProperty))
+		{
+			for (TFieldIterator<FProperty> FieldIt(StructProperty->Struct); FieldIt; ++FieldIt)
+			{
+				if (!AreSnapshotAndOriginalPropertiesEquivalent(WorldData, *FieldIt, SnapshotValuePtr, WorldValuePtr, SnapshotActor, WorldActor))
+				{
+					return false;
+				}
+			}
+			return true;
+		}
+
+		// Check whether property value points to a subobject
+		if (const FObjectPropertyBase* ObjectProperty = CastField<FObjectPropertyBase>(LeafProperty))
+		{
+			return AreObjectPropertiesEquivalent(WorldData, ObjectProperty, SnapshotValuePtr, WorldValuePtr, SnapshotActor, WorldActor);
+		}
+
+
+		// Use our custom equality function for array, set, and map inner properties
+		if (const FArrayProperty* ArrayProperty = CastField<FArrayProperty>(LeafProperty))
+		{
+			FScriptArrayHelper SnapshotArray(ArrayProperty, SnapshotValuePtr);
+			FScriptArrayHelper WorldArray(ArrayProperty, WorldValuePtr);
+			if (SnapshotArray.Num() != WorldArray.Num())
+			{
+				return false;
+			}
+
+			for (int32 j = 0; j < SnapshotArray.Num(); ++j)
+			{
+				void* SnapshotElementValuePtr = SnapshotArray.GetRawPtr(j);
+				void* WorldElementValuePtr = WorldArray.GetRawPtr(j);
+
+				if (!AreSnapshotAndOriginalPropertiesEquivalent(WorldData, ArrayProperty->Inner, SnapshotElementValuePtr, WorldElementValuePtr, SnapshotActor, WorldActor))
+				{
+					return false;
+				}
+			}
+			return true;
+		}
+		
+		if (const FMapProperty* MapProperty = CastField<FMapProperty>(LeafProperty))
+		{
+			// TODO: Use custom function. Need to do something similar to UE4MapProperty_Private::IsPermutation
+		}
+
+		if (const FSetProperty* SetProperty = CastField<FSetProperty>(LeafProperty))
+		{
+			// TODO: Use custom function. Need to do something similar to UE4SetProperty_Private::IsPermutation
+		}
+
+		if (const FTextProperty* TextProperty = CastField<FTextProperty>(LeafProperty))
+		{
+			const FText& SnapshotText = TextProperty->GetPropertyValue_InContainer(SnapshotContainer);
+			const FText& WorldText = TextProperty->GetPropertyValue_InContainer(WorldContainer);
+			return SnapshotText.IdenticalTo(WorldText, ETextIdenticalModeFlags::None) || SnapshotText.ToString().Equals(WorldText.ToString());
+		}
+		
+		// Use normal property comparison for all other properties
+		if (!LeafProperty->Identical_InContainer(SnapshotContainer, WorldContainer, i, PPF_DeepComparison | PPF_DeepCompareDSOsOnly))
+		{
+			return false;
+		}
+	}
+	return true;
+}
+
+bool SnapshotUtil::AreObjectPropertiesEquivalent(const FWorldSnapshotData& WorldData, const FObjectPropertyBase* ObjectProperty, void* SnapshotValuePtr, void* WorldValuePtr, AActor* SnapshotActor, AActor* WorldActor)
+{
+	// Native identity check handles:
+	// - external references, e.g. UMaterial in content browser
+	// - soft object paths: if SnapshotValuePtr is a TSoftObjectPtr<AActor> property, it retains the object path to the editor object.
+	if (ObjectProperty->Identical(SnapshotValuePtr, WorldValuePtr, 0))
+	{
+		return true;
+	}
+		
+	UObject* SnapshotObject = ObjectProperty->GetObjectPropertyValue(SnapshotValuePtr);
+	UObject* WorldObject = ObjectProperty->GetObjectPropertyValue(WorldValuePtr);
+	return AreReferencesEquivalent(WorldData, SnapshotObject, WorldObject, SnapshotActor, WorldActor);
+}
+
+namespace
+{
+	/** Checks whether the two actors object properties should be considered equivalent */
+	bool AreActorsEquivalent(UObject* SnapshotPropertyValue, AActor* OriginalActorReference, const TMap<FSoftObjectPath, FActorSnapshotData>& ActorData)
+	{
+		// Compare actors
+		const FActorSnapshotData* SavedData = ActorData.Find(OriginalActorReference);
+		if (SavedData == nullptr)
+		{
+			return false;
+		}
+
+		// The snapshot actor was already allocated, if some other snapshot actor is referencing it
+		const TOptional<AActor*> PreallocatedSnapshotVersion = SavedData->GetPreallocatedIfValidButDoNotAllocate();
+		return PreallocatedSnapshotVersion.Get(nullptr) == SnapshotPropertyValue;
+	}
+	
+	/** Checks whether the two subobject object properties should be considered equivalent */
+	bool HaveSameNames(UObject* SnapshotPropertyValue, UObject* OriginalPropertyValue, const TMap<FSoftObjectPath, FActorSnapshotData>& ActorData)
+	{
+		AActor* SnapshotOwningActor = SnapshotPropertyValue->GetTypedOuter<AActor>();
+		AActor* OriginalOwningActor = OriginalPropertyValue->GetTypedOuter<AActor>();
+		if (!ensureMsgf(SnapshotOwningActor && OriginalOwningActor, TEXT("This is weird: the objects are part of a world and not actors, so they should be subobjects of actors, like components. Investigate")))
+		{
+			return false;
+		}
+
+		// Are the two subobjects owned by equivalent actors
+		const FActorSnapshotData* EquivalentActorData = ActorData.Find(OriginalOwningActor);
+		check (EquivalentActorData);
+		AActor* EquivalentSnapshotActor = EquivalentActorData->GetPreallocatedIfValidButDoNotAllocate().Get(nullptr);
+		const bool bAreOwnedByEquivalentActors = EquivalentSnapshotActor == SnapshotOwningActor; 
+		if (!bAreOwnedByEquivalentActors)
+		{
+			return false;
+		}
+
+		// Check that chain of outers correspond to each other.
+		UObject* CurrentSnapshotOuter = SnapshotPropertyValue;
+		UObject* CurrentOriginalOuter = OriginalPropertyValue;
+		for (; CurrentSnapshotOuter != SnapshotOwningActor && CurrentOriginalOuter != OriginalOwningActor; CurrentSnapshotOuter = CurrentSnapshotOuter->GetOuter(), CurrentOriginalOuter = CurrentOriginalOuter->GetOuter())
+		{
+			const bool bHaveSameName = CurrentSnapshotOuter->GetFName().IsEqual(CurrentOriginalOuter->GetFName());
+			// I thought of also checking whether the two outers have the same class but I see no reason to atm
+			if (!bHaveSameName)
+			{
+				return false;
+			}
+		}
+
+		return true;
+	}
+
+	bool AreEquivalentSubobjects(const FWorldSnapshotData& WorldData, UObject* SnapshotPropertyValue, UObject* OriginalPropertyValue, AActor* SnapshotActor, AActor* OriginalActor)
+	{
+		AActor* OwningSnapshotActor = SnapshotPropertyValue->GetTypedOuter<AActor>();
+		AActor* OwningWorldActor = OriginalPropertyValue->GetTypedOuter<AActor>();
+		if (ensure(OwningSnapshotActor && OwningWorldActor))
+		{
+			return AreActorsEquivalent(OwningSnapshotActor, OwningWorldActor, WorldData.ActorData) && !HaveDifferentPropertyValues(WorldData, SnapshotPropertyValue, OriginalPropertyValue, SnapshotActor, OriginalActor);
+		}
+		return false;
+	}
+}
+
+bool SnapshotUtil::AreReferencesEquivalent(const FWorldSnapshotData& WorldData, UObject* SnapshotPropertyValue, UObject* OriginalPropertyValue, AActor* SnapshotActor, AActor* OriginalActor)
+{
+	if (SnapshotPropertyValue == nullptr || OriginalPropertyValue == nullptr)
+	{
+		// Migration: We did not save subobject data. In this case snapshot version resolves to nullptr.
+		const bool bWorldObjectIsSubobject = OriginalPropertyValue && !OriginalPropertyValue->IsA<AActor>() && OriginalPropertyValue->IsInA(AActor::StaticClass());
+		const bool bIsOldSnapshot = WorldData.GetSnapshotVersionInfo().GetSnapshotCustomVersion() < FSnapshotCustomVersion::SubobjectSupport;
+		const bool bOldSnapshotDataDidNotCaptureSubobjects = bIsOldSnapshot && bWorldObjectIsSubobject && SnapshotPropertyValue == nullptr;
+		
+		return bOldSnapshotDataDidNotCaptureSubobjects || SnapshotPropertyValue == OriginalPropertyValue;
+	}
+	if (SnapshotPropertyValue->GetClass() != OriginalPropertyValue->GetClass())
+	{
+		return false;
+	}
+
+	if (AActor* OriginalActorReference = Cast<AActor>(OriginalPropertyValue))
+	{
+		return AreActorsEquivalent(SnapshotPropertyValue, OriginalActorReference, WorldData.ActorData);
+	}
+	
+	const bool bIsWorldObject = SnapshotPropertyValue->IsInA(UWorld::StaticClass()) && OriginalPropertyValue->IsInA(UWorld::StaticClass());
+	if (bIsWorldObject)
+	{
+		// Note: Only components are required to have same names. Other subobjects only need to have equal properties.
+		const bool bAreComponents = SnapshotPropertyValue->IsA<UActorComponent>() && OriginalPropertyValue->IsA<UActorComponent>();
+		return (bAreComponents && HaveSameNames(SnapshotPropertyValue, OriginalPropertyValue, WorldData.ActorData))
+			|| !FSnapshotRestorability::IsSubobjectDesirableForCapture(OriginalPropertyValue)
+			|| (!bAreComponents && AreEquivalentSubobjects(WorldData, SnapshotPropertyValue, OriginalPropertyValue, SnapshotActor, OriginalActor));
+	}
+	
+	return SnapshotPropertyValue == OriginalPropertyValue;
+}
\ No newline at end of file
diff --git a/Engine/Plugins/VirtualProduction/LevelSnapshots/Source/LevelSnapshots/Private/Data/Util/EquivalenceUtil.h b/Engine/Plugins/VirtualProduction/LevelSnapshots/Source/LevelSnapshots/Private/Data/Util/EquivalenceUtil.h
new file mode 100644
--- /dev/null
+++ b/Engine/Plugins/VirtualProduction/LevelSnapshots/Source/LevelSnapshots/Private/Data/Util/EquivalenceUtil.h
@@ -0,0 +1,47 @@
+﻿// Copyright Epic Games, Inc. All Rights Reserved.
+
+#pragma once
+
+#include "CoreMinimal.h"
+
+class AActor;
+class UObject;
+struct FWorldSnapshotData;
+
+namespace SnapshotUtil
+{
+	typedef TFunction<void(UActorComponent* SnapshotComponent, UActorComponent* WorldComponent)> FHandleMatchedActorComponent;
+	typedef TFunction<void(UActorComponent*)> FHandleUnmatchedActorComponent;
+
+	/**
+	 * Iterates through both actors' component lists and calls the appropriate callback.
+	 *
+	 * @param SnapshotActor Actor stored in ULevelSnapshot's internal world
+	 * @param WorldActor Actor stored in level viewport world
+	 * @param OnComponentsMatched Called when component exists on both actors
+	 * @param OnSnapshotComponentUnmatched Called when component exists on snapshot but not on editor world
+	 * @param OnWorldComponentUnmatched Called when component exists in editor world but not in snapshot world
+	 */
+	void IterateComponents(AActor* SnapshotActor, AActor* WorldActor, FHandleMatchedActorComponent OnComponentsMatched, FHandleUnmatchedActorComponent OnSnapshotComponentUnmatched, FHandleUnmatchedActorComponent OnWorldComponentUnmatched);
+
+	/** Checks whether the original actor has any properties that changed since the snapshot was taken.  */
+	bool HasOriginalChangedPropertiesSinceSnapshotWasTaken(const FWorldSnapshotData& WorldData, AActor* SnapshotActor, AActor* WorldActor);
+
+	/**
+	 * Checks whether the snapshot and original property value should be considered equal.
+	 * Primitive properties are trivial. Special support is needed for object references.
+	 */
+	bool AreSnapshotAndOriginalPropertiesEquivalent(const FWorldSnapshotData& WorldData, const FProperty* LeafProperty, void* SnapshotContainer, void* WorldContainer, AActor* SnapshotActor, AActor* WorldActor);
+
+	/**
+	 * Two object properties are equivalent if they are
+	 *  - both null
+	 *  - asset references and point to the same asset
+	 *  - actor or component references
+	 *  - subobject references where 1. the subobject have the same name, 2. have the same classes, and 3. the properties are equal
+	 */
+	bool AreObjectPropertiesEquivalent(const FWorldSnapshotData& WorldData, const FObjectPropertyBase* ObjectProperty, void* SnapshotValuePtr, void* WorldValuePtr, AActor* SnapshotActor, AActor* WorldActor);
+	
+	/** Checks whether two pointers point to "equivalent" objects. */
+	bool AreReferencesEquivalent(const FWorldSnapshotData& WorldData, UObject* SnapshotPropertyValue, UObject* OriginalPropertyValue, AActor* SnapshotActor, AActor* OriginalActor);
+}
diff --git a/Engine/Plugins/VirtualProduction/LevelSnapshots/Source/LevelSnapshots/Private/Data/Util/SnapshotObjectUtil.cpp b/Engine/Plugins/VirtualProduction/LevelSnapshots/Source/LevelSnapshots/Private/Data/Util/SnapshotObjectUtil.cpp
new file mode 100644
--- /dev/null
+++ b/Engine/Plugins/VirtualProduction/LevelSnapshots/Source/LevelSnapshots/Private/Data/Util/SnapshotObjectUtil.cpp
@@ -0,0 +1,349 @@
+﻿// Copyright Epic Games, Inc. All Rights Reserved.
+
+#include "Util/SnapshotObjectUtil.h"
+
+#include "ApplySnapshotDataArchiveV2.h"
+#include "Archive/TakeWorldObjectSnapshotArchive.h"
+#include "LevelSnapshotsLog.h"
+#include "SnapshotArchive.h"
+#include "SnapshotUtil.h"
+#include "SubobjectSnapshotData.h"
+#include "WorldSnapshotData.h"
+
+#include "Components/ActorComponent.h"
+#include "EngineUtils.h"
+#include "ObjectTrace.h"
+#include "SnapshotRestorability.h"
+#include "GameFramework/Actor.h"
+#include "Misc/ScopeExit.h"
+#include "Templates/NonNullPointer.h"
+
+namespace
+{
+	UObject* ResolveExternalReference(FWorldSnapshotData& WorldData, const FSoftObjectPath& ObjectPath)
+	{
+		UObject* ExternalReference = ObjectPath.ResolveObject();
+		const bool bNeedsToLoadFromDisk = ExternalReference == nullptr;
+		if (bNeedsToLoadFromDisk)
+		{
+			ExternalReference = ObjectPath.TryLoad();
+		}
+		
+		// We're supposed to be dealing with an external (asset) reference, e.g. a UMaterial.
+		ensureAlwaysMsgf(ExternalReference == nullptr || (!ExternalReference->IsA<AActor>() && !ExternalReference->IsA<UActorComponent>()), TEXT("Something is wrong. We just checked that the reference is not a world object but it is an actor or component."));
+		return ExternalReference;
+	}
+	
+	UObject* CreateSubobjectForSnapshotWorld(FWorldSnapshotData& WorldData, int32 ObjectPathIndex, const FSoftObjectPath& SnapshotPathToSubobject, const FString& LocalisationNamespace)
+	{
+		FSubobjectSnapshotData* SubobjectData = WorldData.Subobjects.Find(ObjectPathIndex);
+		if (!SubobjectData || SubobjectData->bWasBlacklisted)
+		{
+			return nullptr;
+		}
+		if (SubobjectData->SnapshotObject.IsValid())
+		{
+			return SubobjectData->SnapshotObject.Get();
+		}
+
+
+
+	
+		UClass* Class = SubobjectData->Class.ResolveClass();
+		if (!Class)
+		{
+			UE_LOG(LogLevelSnapshots, Warning, TEXT("Class '%s' not found. Maybe it was removed?"), *SubobjectData->Class.ToString());
+			return nullptr;
+		}
+		
+		const int32 OuterIndex = SubobjectData->OuterIndex;
+		check(OuterIndex != ObjectPathIndex);
+		UObject* SubobjectOuter = SnapshotUtil::Object::ResolveObjectDependencyForSnapshotWorld(WorldData, OuterIndex, LocalisationNamespace);
+		if (!SubobjectOuter)
+		{
+			UE_LOG(LogLevelSnapshots, Warning, TEXT("Failed to create '%s' because its outer could not be created."), *SnapshotPathToSubobject.ToString());
+			return nullptr;
+		}
+
+
+	
+		const FName SubobjectName = *SnapshotUtil::ExtractLastSubobjectName(SnapshotPathToSubobject);
+		const bool bSubobjectNameIsTaken = [SubobjectOuter, SubobjectName]()
+		{
+			TArray<UObject*> Subobjects;
+			GetObjectsWithOuter(SubobjectOuter, Subobjects, true);
+			for (UObject* Object : Subobjects)
+			{
+				if (Object->GetFName() == SubobjectName)
+				{
+					return true;
+				}
+			}
+			return false;
+		}();
+		if (bSubobjectNameIsTaken)
+		{
+			UE_LOG(LogLevelSnapshots, Warning, TEXT("Failed to create '%s' because subobject name was already taken."), *SnapshotPathToSubobject.ToString());
+			return nullptr;
+		}
+
+
+
+	
+		
+		UObject* Subobject = NewObject<UObject>(SubobjectOuter, Class, SubobjectName); 
+		SubobjectData->SnapshotObject = Subobject;
+			
+		WorldData.SerializeClassDefaultsInto(Subobject);
+		FSnapshotArchive::ApplyToSnapshotWorldObject(*SubobjectData, WorldData, Subobject, LocalisationNamespace)	;
+
+		return Subobject;
+	}
+	
+	UObject* SerializeOrCreateSubobjectForEditorWorld(FWorldSnapshotData& WorldData, FSubobjectSnapshotData& SubobjectData, UObject* ResolvedObject, const FSoftObjectPath& OriginalObjectPath, const FString& LocalisationNamespace, const FPropertySelectionMap& SelectionMap)
+	{
+		using namespace SnapshotUtil;
+		
+		if (SubobjectData.EditorObject.IsValid())
+		{
+			return SubobjectData.EditorObject.Get();
+		}
+
+		if (SubobjectData.bWasBlacklisted)
+		{
+			return !ResolvedObject || ResolvedObject->IsPendingKill() || ResolvedObject->IsUnreachable() ? nullptr : ResolvedObject;
+		}
+		
+		UClass* ExpectedClass = SubobjectData.Class.ResolveClass();
+		if (!ExpectedClass)
+		{
+			UE_LOG(LogLevelSnapshots, Warning, TEXT("Class '%s' could not be resolved."), *SubobjectData.Class.ToString());
+			return nullptr;
+		}
+
+		// Suppose ActorA and ActorB. ActorA is being serialized before ActorB. ActorA has a reference to a component in ActorB.
+		// That component is instanced and was not recreated yet because, unluckily, we're serializing ActorA before ActorB.
+		if (ExpectedClass->IsChildOf(UActorComponent::StaticClass()))
+		{
+			UE_LOG(LogLevelSnapshots, Error, TEXT("Object depends on a component that must be recreated and is owned by another actor. Support for this is not yet implemented."));
+			return nullptr;
+		}
+
+		// Need to clean rename dead objects otherwise we'll get a name clash when instantiating new object
+		const bool bIsReferencingDeadObject = ResolvedObject && (ResolvedObject->IsPendingKill() || ResolvedObject->IsUnreachable());
+		const bool bIsReferencingDifferentClass = ResolvedObject && !ResolvedObject->GetClass()->IsChildOf(ExpectedClass);
+		if (bIsReferencingDeadObject || bIsReferencingDifferentClass)
+		{
+			const FName NewName = MakeUniqueObjectName(ResolvedObject->GetOuter(), ExpectedClass, *ResolvedObject->GetName().Append(TEXT("_TRASH")));
+			ResolvedObject->Rename(*NewName.ToString());
+			ResolvedObject = nullptr;
+		}
+		
+		if (!ensureMsgf(SubobjectData.SnapshotObject.IsValid(), TEXT("Subobject pointer should have been set. Investigate.")))
+		{
+			return ResolvedObject;
+		}
+		
+		UObject* SnapshotVersion = SubobjectData.SnapshotObject.Get();
+		if (ResolvedObject)
+		{
+			// We can't tell whether ResolvedObject is 1: a normal referenced object or 2: a dead object that just has not been collected yet
+			// If it is referenced from somewhere (case 1), we expect it to be in the selection map.
+			// Otherwise we'll treat it as a dead object that we fully restore.
+			if (SelectionMap.IsSubobjectMarkedForReferenceRestorationOnly(OriginalObjectPath))
+			{
+				// No serialization is required. Just return it.
+			}
+			else if (const FPropertySelection* PropertySelection = SelectionMap.GetObjectSelection(OriginalObjectPath).GetPropertySelection())
+			{
+				FApplySnapshotDataArchiveV2::ApplyToExistingEditorWorldObject(SubobjectData, WorldData, ResolvedObject, SnapshotVersion, SelectionMap, *PropertySelection);
+			}
+			else
+			{
+				// Reuse existing instance but treat it like it was just freshly allocated
+				FApplySnapshotDataArchiveV2::ApplyToRecreatedEditorWorldObject(SubobjectData, WorldData, ResolvedObject, SnapshotVersion, SelectionMap);
+			}
+		}
+		else 
+		{
+			UObject* Outer = Object::ResolveObjectDependencyForEditorWorld(WorldData, SubobjectData.OuterIndex, LocalisationNamespace, SelectionMap);
+			if (!Outer)
+			{
+				UE_LOG(LogLevelSnapshots, Error, TEXT("Failed to resolve object '%s' because its outer '%s' could not be resolved"), *OriginalObjectPath.ToString(), *WorldData.SerializedObjectReferences[SubobjectData.OuterIndex].ToString());
+				return nullptr;
+			}
+			
+			ResolvedObject = NewObject<UObject>(Outer, ExpectedClass, *ExtractLastSubobjectName(OriginalObjectPath));
+			if (ensureMsgf(ResolvedObject, TEXT("Failed to allocate '%s'"), *OriginalObjectPath.ToString()))
+			{
+				FApplySnapshotDataArchiveV2::ApplyToRecreatedEditorWorldObject(SubobjectData, WorldData, ResolvedObject, SnapshotVersion, SelectionMap);
+			}
+		}
+		
+		SubobjectData.EditorObject = ResolvedObject;
+		return ResolvedObject;
+	}
+
+	void AddSubobjectDependencyInternal(FWorldSnapshotData& WorldData, UObject* ReferenceFromOriginalObject, int32 ObjectIndex)
+	{
+		if (!FSnapshotRestorability::IsSubobjectDesirableForCapture(ReferenceFromOriginalObject))
+		{
+			WorldData.Subobjects.Add(ObjectIndex, FSubobjectSnapshotData::MakeBlacklisted());
+			return;
+		}
+		
+		const bool bNeedsToSaveSubobject = !WorldData.Subobjects.Contains(ObjectIndex);
+		if (bNeedsToSaveSubobject)
+		{
+			// Avoid infinite recursion
+			WorldData.Subobjects.Add(ObjectIndex);
+			
+			// Owning actor keeps reference to saved subobjects so they can be recreated faster when snapshot is loaded
+			if (!ReferenceFromOriginalObject->IsA<UActorComponent>())
+			{
+				AActor* OwningActor = ReferenceFromOriginalObject->GetTypedOuter<AActor>();
+				check(OwningActor);
+				WorldData.ActorData.FindOrAdd(OwningActor).OwnedSubobjects.Add(ObjectIndex);
+			}
+
+			
+			// Important: first allocate SubobjectData on stack... 
+			FSubobjectSnapshotData SubobjectData; 
+			UClass* Class = ReferenceFromOriginalObject->GetClass();
+			SubobjectData.Class = Class;
+			SubobjectData.OuterIndex = SnapshotUtil::Object::AddObjectDependency(WorldData, ReferenceFromOriginalObject->GetOuter());
+			// ... because serialisation recursively calls this function. FWorldSnapshotData::Subobjects is possibly reallocated
+			FTakeWorldObjectSnapshotArchive Serializer = FTakeWorldObjectSnapshotArchive::MakeArchiveForSavingWorldObject(SubobjectData, WorldData, ReferenceFromOriginalObject);
+			ReferenceFromOriginalObject->Serialize(Serializer);
+
+			
+			WorldData.Subobjects[ObjectIndex] = MoveTemp(SubobjectData); // MoveTemp not profiled
+			WorldData.AddClassDefault(Class);
+		}
+	}
+	
+	void EnsureSubobjectIsSerialized(FWorldSnapshotData& WorldData, UObject* ExistingObject, int32 ObjectPathIndex, const FString& LocalisationNamespace)
+	{
+		FSubobjectSnapshotData* SubobjectData = WorldData.Subobjects.Find(ObjectPathIndex);
+		if (!SubobjectData || SubobjectData->bWasBlacklisted)
+		{
+			return;
+		}
+
+		if (SubobjectData->Class.ResolveClass() != ExistingObject->GetClass())
+		{
+			UE_LOG(LogLevelSnapshots, Warning, TEXT("Skipping serialisation of subobject because classes are different. Snapshot object '%s' will not contain values saved in snapshot."), *ExistingObject->GetName());
+			return;
+		}
+
+		if (!SubobjectData->SnapshotObject.IsValid())
+		{
+			SubobjectData->SnapshotObject = ExistingObject;
+			FSnapshotArchive::ApplyToSnapshotWorldObject(*SubobjectData, WorldData, ExistingObject, LocalisationNamespace);
+		}
+	}
+}
+
+UObject* SnapshotUtil::Object::ResolveObjectDependencyForSnapshotWorld(FWorldSnapshotData& WorldData, int32 ObjectPathIndex, const FString& LocalisationNamespace)
+{
+	using namespace SnapshotUtil;
+	
+	if (!ensure(WorldData.SerializedObjectReferences.IsValidIndex(ObjectPathIndex)))
+	{
+		return nullptr;
+	}
+	
+	const FSoftObjectPath& OriginalObjectPath = WorldData.SerializedObjectReferences[ObjectPathIndex];
+	bool bIsPathToActorSubobject;
+	const TOptional<FActorSnapshotData*> SnapshotData = FindSavedActorDataUsingObjectPath(WorldData.ActorData, OriginalObjectPath, bIsPathToActorSubobject);
+	
+	const bool bIsExternalAssetReference = !SnapshotData.IsSet();
+	if (bIsExternalAssetReference)
+	{
+		return ResolveExternalReference(WorldData, OriginalObjectPath);
+	}
+
+	TOptional<AActor*> Result = SnapshotData.GetValue()->GetPreallocated(WorldData.TempActorWorld->GetWorld(), WorldData);
+	if (!Result.IsSet())
+	{
+		return nullptr;
+	}
+
+	// Reference to actor
+	AActor* SnapshotActor = *Result;
+	const bool bIsPathToActor = !bIsPathToActorSubobject; 
+	if (bIsPathToActor)
+	{
+		return SnapshotActor;
+	}
+
+	// Reference to default subobjects, such as components, or an object we already allocated.
+	const FSoftObjectPath SnapshotPathToSubobject = SetActorInPath(SnapshotActor, OriginalObjectPath);
+	if (UObject* ExistingSubobject = SnapshotPathToSubobject.ResolveObject())
+	{
+		EnsureSubobjectIsSerialized(WorldData, ExistingSubobject, ObjectPathIndex, LocalisationNamespace);
+		return ExistingSubobject;
+	}
+
+	// The subobject is instanced: recreate it recursively
+	return CreateSubobjectForSnapshotWorld(WorldData, ObjectPathIndex, SnapshotPathToSubobject, LocalisationNamespace);
+}
+
+UObject* SnapshotUtil::Object::ResolveObjectDependencyForEditorWorld(FWorldSnapshotData& WorldData, int32 ObjectPathIndex, const FString& LocalisationNamespace, const FPropertySelectionMap& SelectionMap)
+{
+	if (!ensure(WorldData.SerializedObjectReferences.IsValidIndex(ObjectPathIndex)))
+	{
+		return nullptr;
+	}
+	
+	const FSoftObjectPath& OriginalObjectPath = WorldData.SerializedObjectReferences[ObjectPathIndex];
+	bool bIsPathToActorSubobject;
+	const TOptional<FActorSnapshotData*> SnapshotData = FindSavedActorDataUsingObjectPath(WorldData.ActorData, OriginalObjectPath, bIsPathToActorSubobject);
+	
+	// Dealing with an external (asset) reference, e.g. a UMaterial?
+	const bool bIsExternalAssetReference = !SnapshotData.IsSet();
+	if (bIsExternalAssetReference)
+	{
+		return ResolveExternalReference(WorldData, OriginalObjectPath);
+	}
+
+	// Can return immediately for actors and components: their serialization is handled by FActorSnapshotData.
+	UObject* ResolvedObject = OriginalObjectPath.ResolveObject();
+	if (ResolvedObject && (ResolvedObject->IsA<AActor>() || ResolvedObject->IsA<UActorComponent>()))
+	{
+		// Might resolve to a dead object which still exists in memory
+		return ResolvedObject->IsPendingKill() ? nullptr : ResolvedObject;
+	}
+	
+	if (FSubobjectSnapshotData* const SubobjectData = WorldData.Subobjects.Find(ObjectPathIndex))
+	{
+		return SerializeOrCreateSubobjectForEditorWorld(WorldData, *SubobjectData, ResolvedObject, OriginalObjectPath, LocalisationNamespace, SelectionMap);
+	}
+	
+	// This should not happen after release of 4.27. This is error may happen if certain snapshots in an older data format are applied because the data wasn't yet captured.
+	UE_LOG(LogLevelSnapshots, Warning, TEXT("Trying to resolve a subobject but found no associated object data."));
+	return ResolvedObject;
+}
+
+UObject* SnapshotUtil::Object::ResolveObjectDependencyForClassDefaultObject(FWorldSnapshotData& WorldData, int32 ObjectPathIndex)
+{
+	if (ensure(WorldData.SerializedObjectReferences.IsValidIndex(ObjectPathIndex)))
+	{
+		const FSoftObjectPath& OriginalObjectPath = WorldData.SerializedObjectReferences[ObjectPathIndex];
+		return ResolveExternalReference(WorldData, OriginalObjectPath);
+	}
+	return nullptr;
+}
+
+int32 SnapshotUtil::Object::AddObjectDependency(FWorldSnapshotData& WorldData, UObject* ReferenceFromOriginalObject, bool bCheckWhetherSubobject)
+{
+	// Even if FSnapshotRestorability::IsSubobjectDesirableForCapture later returns false for this object, we want to track it
+	const int32 Result = WorldData.SerializedObjectReferences.AddUnique(ReferenceFromOriginalObject);
+
+	if (bCheckWhetherSubobject && ReferenceFromOriginalObject && ReferenceFromOriginalObject->GetTypedOuter<AActor>())
+	{
+		AddSubobjectDependencyInternal(WorldData, ReferenceFromOriginalObject, Result);
+	}
+	
+	return Result;
+}
\ No newline at end of file
diff --git a/Engine/Plugins/VirtualProduction/LevelSnapshots/Source/LevelSnapshots/Private/Data/Util/SnapshotObjectUtil.h b/Engine/Plugins/VirtualProduction/LevelSnapshots/Source/LevelSnapshots/Private/Data/Util/SnapshotObjectUtil.h
new file mode 100644
--- /dev/null
+++ b/Engine/Plugins/VirtualProduction/LevelSnapshots/Source/LevelSnapshots/Private/Data/Util/SnapshotObjectUtil.h
@@ -0,0 +1,41 @@
+﻿// Copyright Epic Games, Inc. All Rights Reserved.
+
+#pragma once
+
+#include "CoreMinimal.h"
+
+class AActor;
+class UObject;
+struct FPropertySelectionMap;
+struct FWorldSnapshotData;
+
+namespace SnapshotUtil
+{
+	namespace Object
+	{
+		/* Resolves an object dependency for use in the snapshot world. If the object is a subobject, it gets fully serialized. */
+		UObject* ResolveObjectDependencyForSnapshotWorld(FWorldSnapshotData& WorldData, int32 ObjectPathIndex, const FString& LocalisationNamespace);
+	
+		/* Resolves an object dependency for use in the editor world. If the object is a subobject, it is serialized.
+		 * Steps for serializing subobject:
+		 *  - If an equivalent object with the saved name and class exists, use that and serialize the properties in SelectionMap into it.
+		 *  - If object does not exist, allocate a new object and serialize all properties into it
+		 *  - If object exists but has different class, skip. Return null.
+		 */
+		UObject* ResolveObjectDependencyForEditorWorld(FWorldSnapshotData& WorldData, int32 ObjectPathIndex, const FString& LocalisationNamespace, const FPropertySelectionMap& SelectionMap);
+
+		/* Resolves an object depedency when restoring a class default object. Simply resolves without further checks. */
+		UObject* ResolveObjectDependencyForClassDefaultObject(FWorldSnapshotData& WorldData, int32 ObjectPathIndex);
+
+
+		
+		/**
+		 * Saves the object's path. Intended for external objects, e.g. to UMaterial in the content browser.
+		 *
+		 * @param ReferenceFromOriginalObject The reference as it was found in a property in the editor world
+		 * @param bCheckWhetherSubobject If true, this will check the the object is a subobject of an actor and call AddSubobjectDependency.
+		 */
+		int32 AddObjectDependency(FWorldSnapshotData& WorldData, UObject* ReferenceFromOriginalObject, bool bCheckWhetherSubobject = true);
+	}
+}
+
diff --git a/Engine/Plugins/VirtualProduction/LevelSnapshots/Source/LevelSnapshots/Private/Data/Util/SnapshotUtil.cpp b/Engine/Plugins/VirtualProduction/LevelSnapshots/Source/LevelSnapshots/Private/Data/Util/SnapshotUtil.cpp
new file mode 100644
--- /dev/null
+++ b/Engine/Plugins/VirtualProduction/LevelSnapshots/Source/LevelSnapshots/Private/Data/Util/SnapshotUtil.cpp
@@ -0,0 +1,86 @@
+﻿// Copyright Epic Games, Inc. All Rights Reserved.
+
+#include "Util/SnapshotUtil.h"
+
+#include "ActorSnapshotData.h"
+#include "LevelSnapshotsLog.h"
+
+#include "EngineUtils.h"
+#include "GameFramework/Actor.h"
+#include "Misc/ScopeExit.h"
+
+FString SnapshotUtil::ExtractLastSubobjectName(const FSoftObjectPath& ObjectPath)
+{
+	const FString& SubPathString = ObjectPath.GetSubPathString();
+	const int32 LastDotIndex = SubPathString.Find(TEXT("."), ESearchCase::IgnoreCase, ESearchDir::FromEnd);
+	if (LastDotIndex == INDEX_NONE)
+	{
+		return SubPathString;
+	}
+
+	return SubPathString.RightChop(LastDotIndex + 1);
+}
+
+/* If the path contains an actor, returns a new path to that actor. */
+TOptional<FSoftObjectPath> SnapshotUtil::ExtractActorFromPath(const FSoftObjectPath& OriginalObjectPath, bool& bIsPathToActorSubobject)
+{
+	const static FString PersistentLevelString("PersistentLevel.");
+	const int32 PersistentLevelStringLength = PersistentLevelString.Len();
+	const FString& SubPathString = OriginalObjectPath.GetSubPathString();
+	const int32 IndexOfPersistentLevelInfo = SubPathString.Find(PersistentLevelString, ESearchCase::CaseSensitive);
+	if (IndexOfPersistentLevelInfo == INDEX_NONE)
+	{
+		return {};
+	}
+
+	const int32 DotAfterActorNameIndex = SubPathString.Find(TEXT("."), ESearchCase::IgnoreCase, ESearchDir::FromStart, IndexOfPersistentLevelInfo + PersistentLevelStringLength);
+	const bool bPathIsToActor = DotAfterActorNameIndex == INDEX_NONE;
+	
+	bIsPathToActorSubobject = !bPathIsToActor;
+	return bPathIsToActor ?
+		// Example /Game/MapName.MapName:PersistentLevel.StaticMeshActor_42
+		OriginalObjectPath
+		:
+		// Converts /Game/MapName.MapName:PersistentLevel.StaticMeshActor_42.StaticMeshComponent to /Game/MapName.MapName:PersistentLevel.StaticMeshActor_42
+		FSoftObjectPath(OriginalObjectPath.GetAssetPathName(), SubPathString.Left(DotAfterActorNameIndex));
+}
+
+/* If Path contains a path to an actor, returns that actor.
+ * Example: /Game/MapName.MapName:PersistentLevel.StaticMeshActor_42.StaticMeshComponent returns StaticMeshActor_42's data
+ */
+TOptional<FActorSnapshotData*> SnapshotUtil::FindSavedActorDataUsingObjectPath(TMap<FSoftObjectPath, FActorSnapshotData>& ActorData, const FSoftObjectPath& OriginalObjectPath, bool& bIsPathToActorSubobject)
+{
+	const TOptional<FSoftObjectPath> PathToActor = SnapshotUtil::ExtractActorFromPath(OriginalObjectPath, bIsPathToActorSubobject);
+	if (!PathToActor.IsSet())
+	{
+		return {};	
+	}
+	
+	FActorSnapshotData* Result = ActorData.Find(*PathToActor);
+	UE_CLOG(Result == nullptr, LogLevelSnapshots, Warning, TEXT("Path %s looks like an actor path but no data was saved for it. Maybe it was a reference to an auto-generated actor, e.g. a brush or volume present in all worlds by default?"), *OriginalObjectPath.ToString());
+	return Result ? Result : TOptional<FActorSnapshotData*>();  
+}
+
+
+/* Takes an existing path to an actor's subobjects and replaces the actor bit with the path to another actor.
+ *
+ * E.g. /Game/MapName.MapName:PersistentLevel.StaticMeshActor_42.StaticMeshComponent could become /Game/MapName.MapName:PersistentLevel.SomeOtherActor.StaticMeshComponent
+ */
+FSoftObjectPath SnapshotUtil::SetActorInPath(AActor* NewActor, const FSoftObjectPath& OriginalObjectPath)
+{
+	const static FString PersistentLevelString("PersistentLevel.");
+	const int32 PersistentLevelStringLength = PersistentLevelString.Len();
+	const FString& SubPathString = OriginalObjectPath.GetSubPathString();
+	const int32 IndexOfPersistentLevelInfo = SubPathString.Find(PersistentLevelString, ESearchCase::CaseSensitive);
+	if (IndexOfPersistentLevelInfo == INDEX_NONE)
+	{
+		return {};
+	}
+
+	const int32 DotAfterActorNameIndex = SubPathString.Find(TEXT("."), ESearchCase::IgnoreCase, ESearchDir::FromStart, IndexOfPersistentLevelInfo + PersistentLevelStringLength);
+
+	const FSoftObjectPath PathToNewActor(NewActor);
+	// PersistentLevel.StaticMeshActor_42.StaticMeshComponent becomes .StaticMeshComponent
+	const FString PathAfterOriginalActor = SubPathString.Right(SubPathString.Len() - DotAfterActorNameIndex); 
+	return FSoftObjectPath(PathToNewActor.GetAssetPathName(), PathToNewActor.GetSubPathString() + PathAfterOriginalActor);
+}
\ No newline at end of file
diff --git a/Engine/Plugins/VirtualProduction/LevelSnapshots/Source/LevelSnapshots/Private/Data/Util/SnapshotUtil.h b/Engine/Plugins/VirtualProduction/LevelSnapshots/Source/LevelSnapshots/Private/Data/Util/SnapshotUtil.h
new file mode 100644
--- /dev/null
+++ b/Engine/Plugins/VirtualProduction/LevelSnapshots/Source/LevelSnapshots/Private/Data/Util/SnapshotUtil.h
@@ -0,0 +1,30 @@
+﻿// Copyright Epic Games, Inc. All Rights Reserved.
+
+#pragma once
+
+#include "CoreMinimal.h"
+
+class AActor;
+struct FActorSnapshotData;
+
+namespace SnapshotUtil
+{
+	/** Finds the last subobject name in the path */
+	FString ExtractLastSubobjectName(const FSoftObjectPath& ObjectPath);
+
+	/** If path contains an actor, return subobject path to that actor. */
+	TOptional<FSoftObjectPath> ExtractActorFromPath(const FSoftObjectPath& OriginalObjectPath, bool& bIsPathToActorSubobject);
+
+	/**
+	 * If Path contains a path to an actor, returns that actor.
+	 * Example: /Game/MapName.MapName:PersistentLevel.StaticMeshActor_42.StaticMeshComponent returns StaticMeshActor_42's data
+	 */
+	TOptional<FActorSnapshotData*> FindSavedActorDataUsingObjectPath(TMap<FSoftObjectPath, FActorSnapshotData>& ActorData, const FSoftObjectPath& OriginalObjectPath, bool& bIsPathToActorSubobject);
+
+	/**
+	 * Takes an existing path to an actor's subobjects and replaces the actor bit with the path to another actor.
+	 *
+	 * E.g. /Game/MapName.MapName:PersistentLevel.StaticMeshActor_42.StaticMeshComponent could become /Game/MapName.MapName:PersistentLevel.SomeOtherActor.StaticMeshComponent
+	 */
+	FSoftObjectPath SetActorInPath(AActor* NewActor, const FSoftObjectPath& OriginalObjectPath);
+};
diff --git a/Engine/Plugins/VirtualProduction/LevelSnapshots/Source/LevelSnapshots/Private/Data/WorldSnapshotData.cpp b/Engine/Plugins/VirtualProduction/LevelSnapshots/Source/LevelSnapshots/Private/Data/WorldSnapshotData.cpp
--- a/Engine/Plugins/VirtualProduction/LevelSnapshots/Source/LevelSnapshots/Private/Data/WorldSnapshotData.cpp
+++ b/Engine/Plugins/VirtualProduction/LevelSnapshots/Source/LevelSnapshots/Private/Data/WorldSnapshotData.cpp
@@ -12,7 +12,10 @@
 #include "Components/ActorComponent.h"
 #include "EngineUtils.h"
 #include "GameFramework/Actor.h"
+#include "Misc/ScopeExit.h"
 #include "UObject/Package.h"
+#include "Util/SnapshotObjectUtil.h"
+#include "Util/SnapshotUtil.h"
 #if WITH_EDITOR
 #include "Editor.h"
 #include "Editor/EditorEngine.h"
@@ -25,144 +28,6 @@
 
 #define LOCTEXT_NAMESPACE "LevelSnapshotsEditor"
 
-namespace
-{
-	/* If the path contains an actor, returns a new path to that actor. */
-	TOptional<FSoftObjectPath> ExtractActorFromPath(const FSoftObjectPath& OriginalObjectPath, bool& bIsPathToActorSubobject)
-	{
-		const static FString PersistentLevelString("PersistentLevel.");
-		const int32 PersistentLevelStringLength = PersistentLevelString.Len();
-		const FString& SubPathString = OriginalObjectPath.GetSubPathString();
-		const int32 IndexOfPersistentLevelInfo = SubPathString.Find(PersistentLevelString, ESearchCase::CaseSensitive);
-		if (IndexOfPersistentLevelInfo == INDEX_NONE)
-		{
-			return {};
-		}
-
-		const int32 DotAfterActorNameIndex = SubPathString.Find(TEXT("."), ESearchCase::IgnoreCase, ESearchDir::FromStart, IndexOfPersistentLevelInfo + PersistentLevelStringLength);
-		const bool bPathIsToActor = DotAfterActorNameIndex == INDEX_NONE;
-		
-		bIsPathToActorSubobject = !bPathIsToActor;
-		return bPathIsToActor ?
-			// Example /Game/MapName.MapName:PersistentLevel.StaticMeshActor_42
-			OriginalObjectPath
-			:
-			// Converts /Game/MapName.MapName:PersistentLevel.StaticMeshActor_42.StaticMeshComponent to /Game/MapName.MapName:PersistentLevel.StaticMeshActor_42
-			FSoftObjectPath(OriginalObjectPath.GetAssetPathName(), SubPathString.Left(DotAfterActorNameIndex));
-	}
-	
-	/* If Path contains a path to an actor, returns that actor.
-	 * Example: /Game/MapName.MapName:PersistentLevel.StaticMeshActor_42.StaticMeshComponent returns StaticMeshActor_42's data
-	 */
-	TOptional<FActorSnapshotData*> FindSavedActorDataUsingObjectPath(TMap<FSoftObjectPath, FActorSnapshotData>& ActorData, const FSoftObjectPath& OriginalObjectPath, bool& bIsPathToActorSubobject)
-	{
-		const TOptional<FSoftObjectPath> PathToActor = ExtractActorFromPath(OriginalObjectPath, bIsPathToActorSubobject);
-		if (!PathToActor.IsSet())
-		{
-			return {};	
-		}
-		
-		FActorSnapshotData* Result = ActorData.Find(*PathToActor);
-		UE_CLOG(Result == nullptr, LogLevelSnapshots, Warning, TEXT("Path %s looks like an actor path but no data was saved for it. Maybe it was a reference to an auto-generated actor, e.g. a brush or volume present in all worlds by default?"), *OriginalObjectPath.ToString());
-		return Result ? Result : TOptional<FActorSnapshotData*>();  
-	}
-
-
-	/* Takes an existing path to an actor's subobjects and replaces the actor bit with the path to another actor.
-	 *
-	 * E.g. /Game/MapName.MapName:PersistentLevel.StaticMeshActor_42.StaticMeshComponent could become /Game/MapName.MapName:PersistentLevel.SomeOtherActor.StaticMeshComponent
-	 */
-	FSoftObjectPath SetActorInPath(AActor* NewActor, const FSoftObjectPath& OriginalObjectPath)
-	{
-		const static FString PersistentLevelString("PersistentLevel.");
-		const int32 PersistentLevelStringLength = PersistentLevelString.Len();
-		const FString& SubPathString = OriginalObjectPath.GetSubPathString();
-		const int32 IndexOfPersistentLevelInfo = SubPathString.Find(PersistentLevelString, ESearchCase::CaseSensitive);
-		if (IndexOfPersistentLevelInfo == INDEX_NONE)
-		{
-			return {};
-		}
-
-		const int32 DotAfterActorNameIndex = SubPathString.Find(TEXT("."), ESearchCase::IgnoreCase, ESearchDir::FromStart, IndexOfPersistentLevelInfo + PersistentLevelStringLength);
-
-		const FSoftObjectPath PathToNewActor(NewActor);
-		// PersistentLevel.StaticMeshActor_42.StaticMeshComponent becomes .StaticMeshComponent
-		const FString PathAfterOriginalActor = SubPathString.Right(SubPathString.Len() - DotAfterActorNameIndex); 
-		return FSoftObjectPath(PathToNewActor.GetAssetPathName(), PathToNewActor.GetSubPathString() + PathAfterOriginalActor);
-	}
-
-	enum EEquivalenceResult
-	{
-		Equivalent,
-		NotEquivalent,
-		NotComparable
-	};
-	
-	EEquivalenceResult AreActorsEquivalent(UObject* SnapshotPropertyValue, UObject* OriginalPropertyValue, const TMap<FSoftObjectPath, FActorSnapshotData>& ActorData)
-	{
-		// Compare actors
-		if (AActor* OriginalActorReference = Cast<AActor>(OriginalPropertyValue))
-		{
-			const FActorSnapshotData* SavedData = ActorData.Find(OriginalActorReference);
-			if (SavedData == nullptr)
-			{
-				return NotEquivalent;
-			}
-
-			// The snapshot actor was already allocated, if some other snapshot actor is referencing it
-			const TOptional<AActor*> PreallocatedSnapshotVersion = SavedData->GetPreallocatedIfValidButDoNotAllocate();
-			return PreallocatedSnapshotVersion.Get(nullptr) == SnapshotPropertyValue ? Equivalent : NotEquivalent;
-		}
-
-		return NotComparable;
-	}
-
-	EEquivalenceResult AreSubobjectsEquivalent(UObject* SnapshotPropertyValue, UObject* OriginalPropertyValue, const TMap<FSoftObjectPath, FActorSnapshotData>& ActorData)
-	{
-		const bool bIsWorldObject = SnapshotPropertyValue->IsInA(UWorld::StaticClass());
-		if (bIsWorldObject)
-		{
-			AActor* SnapshotOwningActor = SnapshotPropertyValue->GetTypedOuter<AActor>();
-			AActor* OriginalOwningActor = OriginalPropertyValue->GetTypedOuter<AActor>();
-			if (!ensureMsgf(SnapshotOwningActor && OriginalOwningActor, TEXT("This is weird: the objects are part of a world and not actors, so they should be subobjects of actors, like components. Investigate")))
-			{
-				return NotEquivalent;
-			}
-
-			// Are the two subobjects owned by corresponding actors
-			const FActorSnapshotData* CorrespondingSnapshotActor = ActorData.Find(OriginalOwningActor);
-			const bool bAreOwnedByEquivalentActors = CorrespondingSnapshotActor == nullptr || CorrespondingSnapshotActor->GetPreallocatedIfValidButDoNotAllocate().Get(nullptr) != SnapshotOwningActor; 
-			if (!bAreOwnedByEquivalentActors)
-			{
-				return NotEquivalent;
-			}
-
-			// TODO: Call registered callbacks from external modules to determine whether the two components match each other.
-			// Needed for cases like Foliage where components are matched based on foliage type instead of name
-
-			// Check that chain of outers correspond to each other.
-			UObject* CurrentSnapshotOuter = SnapshotPropertyValue;
-			UObject* CurrentOriginalOuter = OriginalPropertyValue;
-			for (; CurrentSnapshotOuter != SnapshotOwningActor && CurrentOriginalOuter != OriginalOwningActor; CurrentSnapshotOuter = CurrentSnapshotOuter->GetOuter(), CurrentOriginalOuter = CurrentOriginalOuter->GetOuter())
-			{
-				// TODO: Call registered callbacks from external modules to determine whether the two components match each other.
-				// Needed for cases like Foliage where components are matched based on foliage type instead of name
-
-				const bool bHaveSameName = CurrentSnapshotOuter->GetFName().IsEqual(CurrentOriginalOuter->GetFName());
-				// I thought of also checking whether the two outers have the same class but I see no reason to atm
-				if (!bHaveSameName)
-				{
-					return NotEquivalent;
-				}
-			}
-
-			return Equivalent;
-		}
-		
-		return NotComparable;
-	}
-}
-
 void FWorldSnapshotData::OnCreateSnapshotWorld(UWorld* NewTempActorWorld)
 {
 	TempActorWorld = NewTempActorWorld;
@@ -177,6 +42,12 @@ void FWorldSnapshotData::OnDestroySnapshotWorld()
 	{
 		ClassDefaultsIt->Value.CachedLoadedClassDefault = nullptr;
 	}
+
+	for (auto SubobjectIt = Subobjects.CreateIterator(); SubobjectIt; ++SubobjectIt)
+	{
+		SubobjectIt->Value.SnapshotObject.Reset();
+		SubobjectIt->Value.EditorObject.Reset();
+	}
 }
 
 void FWorldSnapshotData::SnapshotWorld(UWorld* World)
@@ -216,6 +87,12 @@ void FWorldSnapshotData::ApplyToWorld(UWorld* WorldToApplyTo, UPackage* Localisa
 		return;
 	}
 
+	// Clear editor world subobject cache from previous ApplyToWorld
+	for (auto SubobjectIt = Subobjects.CreateIterator(); SubobjectIt; ++SubobjectIt)
+	{
+		SubobjectIt->Value.EditorObject.Reset();
+	}
+	
 	const TArray<FSoftObjectPath> SelectedPaths = PropertiesToSerialize.GetKeys();
 #if WITH_EDITOR
 	const int32 NumActorsToRecreate = PropertiesToSerialize.GetDeletedActorsToRespawn().Num();
@@ -242,7 +119,7 @@ void FWorldSnapshotData::ApplyToWorld(UWorld* WorldToApplyTo, UPackage* Localisa
 
 
 	
-	// If we're in the editor then update the gizmos locations as they can get out of sync if any of the deserialized actors were selected
+	// If we're in the editor then update the gizmos locations as they can get out of sync if any of the serialized actors were selected
 #if WITH_EDITOR
 	if (GUnrealEd)
 	{
@@ -269,190 +146,33 @@ bool FWorldSnapshotData::HasMatchingSavedActor(const FSoftObjectPath& OriginalOb
 	return ActorData.Contains(OriginalObjectPath);
 }
 
-TOptional<AActor*> FWorldSnapshotData::GetDeserializedActor(const FSoftObjectPath& OriginalObjectPath, UPackage* LocalisationSnapshotPackage)
+TOptional<AActor*> FWorldSnapshotData::GetPreallocatedActor(const FSoftObjectPath& OriginalObjectPath)
 {
 	FActorSnapshotData* SerializedActor = ActorData.Find(OriginalObjectPath);
 	if (SerializedActor)
 	{
-		return SerializedActor->GetDeserialized(TempActorWorld->GetWorld(), *this, LocalisationSnapshotPackage);
+		return SerializedActor->GetPreallocated(TempActorWorld->GetWorld(), *this);
 	}
 
-	UE_LOG(LogLevelSnapshots, Warning, TEXT("No save data found for actor %s"), *OriginalObjectPath.ToString());
 	return {};
 }
 
-TOptional<FObjectSnapshotData*> FWorldSnapshotData::GetSerializedClassDefaults(UClass* Class)
-{
-	FObjectSnapshotData* ClassDefaultData = ClassDefaults.Find(Class);
-	return ClassDefaultData ? ClassDefaultData : TOptional<FObjectSnapshotData*>();
-}
-
-bool FWorldSnapshotData::AreReferencesEquivalent(UObject* SnapshotPropertyValue, UObject* OriginalPropertyValue, AActor* SnapshotActor, AActor* OriginalActor) const
-{
-	const bool bOriginalIsSubobject = OriginalPropertyValue != nullptr && OriginalPropertyValue->IsIn(OriginalActor);
-	const bool bNeedsUnsupportedSubobjectRestorationFeature = SnapshotPropertyValue == nullptr && bOriginalIsSubobject;
-	if (bNeedsUnsupportedSubobjectRestorationFeature)
-	{
-		UE_LOG(LogLevelSnapshots, Verbose, TEXT("Object '%s' in of actor '%s' seems to be a subobject. Snapshots currently do not support re-creating subobjects."),
-			OriginalPropertyValue ? *OriginalPropertyValue->GetName() : TEXT("null"),
-			*OriginalActor->GetName()
-			);
-		return true;
-	}
-	
-	if (SnapshotPropertyValue == nullptr || OriginalPropertyValue == nullptr)
-	{
-		return SnapshotPropertyValue == OriginalPropertyValue;
-	}
-	if (SnapshotPropertyValue->GetClass() != OriginalPropertyValue->GetClass())
-	{
-		return false;
-	}
-
-	const EEquivalenceResult ActorsEquivalent = AreActorsEquivalent(SnapshotPropertyValue, OriginalPropertyValue, ActorData);
-	if (ActorsEquivalent != NotComparable)
-	{
-		return ActorsEquivalent == Equivalent;
-	}
-
-	const EEquivalenceResult SubobjectsEquivalent = AreSubobjectsEquivalent(SnapshotPropertyValue, OriginalPropertyValue, ActorData);
-	if (SubobjectsEquivalent != NotComparable)
-	{
-		return SubobjectsEquivalent == Equivalent;
-	}
-	
-	// Compare external references, like a UMaterial in the content browser.
-	return SnapshotPropertyValue == OriginalPropertyValue;
-}
-
-int32 FWorldSnapshotData::AddObjectDependency(UObject* ReferenceFromOriginalObject)
-{
-	return SerializedObjectReferences.AddUnique(ReferenceFromOriginalObject);
-}
-
-UObject* FWorldSnapshotData::ResolveObjectDependencyForSnapshotWorld(int32 ObjectPathIndex)
-{
-	if (!ensure(SerializedObjectReferences.IsValidIndex(ObjectPathIndex)))
-	{
-		return nullptr;
-	}
-	
-	const FSoftObjectPath& OriginalObjectPath = SerializedObjectReferences[ObjectPathIndex];
-	bool bIsPathToActorSubobject;
-	const TOptional<FActorSnapshotData*> SnapshotData = FindSavedActorDataUsingObjectPath(ActorData, OriginalObjectPath, bIsPathToActorSubobject);
-	
-	const bool bIsExternalAssetReference = !SnapshotData.IsSet();
-	if (bIsExternalAssetReference)
-	{
-		return ResolveExternalReference(OriginalObjectPath);
-	}
-
-	TOptional<AActor*> Result = SnapshotData.GetValue()->GetPreallocated(TempActorWorld->GetWorld(), *this);
-	if (!Result.IsSet())
-	{
-		return nullptr;
-	}
-
-	// Reference to actor
-	AActor* SnapshotActor = *Result;
-	const bool bIsPathToActor = !bIsPathToActorSubobject; 
-	if (bIsPathToActor)
-	{
-		return SnapshotActor;
-	}
-
-	// Reference to default subobjects, such as components, or an object we already allocated.
-	const FSoftObjectPath PathToSubobject = SetActorInPath(SnapshotActor, OriginalObjectPath);
-	if (UObject* DefaultSubobject = PathToSubobject.ResolveObject())
-	{
-		return DefaultSubobject;
-	}
-
-	// TODO: The subobject is instanced. Recreate it here using Subobjects.Find(ObjectPathIndex) to recursively construct the subobject's outers
-	return nullptr;
-}
-
-UObject* FWorldSnapshotData::ResolveObjectDependencyForEditorWorld(int32 ObjectPathIndex, const FPropertySelectionMap& SelectionMap)
+TOptional<AActor*> FWorldSnapshotData::GetDeserializedActor(const FSoftObjectPath& OriginalObjectPath, UPackage* LocalisationSnapshotPackage)
 {
-	if (!ensure(SerializedObjectReferences.IsValidIndex(ObjectPathIndex)))
-	{
-		return nullptr;
-	}
-	
-	const FSoftObjectPath& OriginalObjectPath = SerializedObjectReferences[ObjectPathIndex];
-
-	bool bIsPathToActorSubobject;
-	const TOptional<FActorSnapshotData*> SnapshotData = FindSavedActorDataUsingObjectPath(ActorData, OriginalObjectPath, bIsPathToActorSubobject);
-	
-	// Dealing with an external (asset) reference, e.g. a UMaterial?
-	const bool bIsExternalAssetReference = !SnapshotData.IsSet();
-	if (bIsExternalAssetReference)
-	{
-		return ResolveExternalReference(OriginalObjectPath);
-	}
-
-	UObject* ResolvedObject = OriginalObjectPath.ResolveObject();
-	// Actor reference?
-	const bool bIsReferenceToActor = !bIsPathToActorSubobject;
-	if (bIsReferenceToActor)
-	{
-		return ResolvedObject;
-	}
-
-
-	const FSubobjectSnapshotData* SubobjectData = Subobjects.Find(ObjectPathIndex);
-	if (!SubobjectData)
-	{
-		// This should not happen after release of 4.27. This is error may happen if certain snapshots in an older data format are applied because the data wasn't yet captured.
-		UE_LOG(LogLevelSnapshots, Warning, TEXT("Trying to resolve a subobject but found no associated object data."));
-		return ResolvedObject;
-	}
-	
-	// TODO: Implement restoring subobjects, e.g. material instances, here
-	UClass* ExpectedClass = SubobjectData->Class.ResolveClass();
-	if (ResolvedObject == nullptr || !ensure(ExpectedClass) || ExpectedClass != ResolvedObject->GetClass())
-	{
-		// There was already a subobject with the same name but it is incompatible because is a different class.
-		// TODO: Allocate a new instance, serialize all data into it, and purge ResolvedObject by replacing all references to it
-	}
-	else
+	FActorSnapshotData* SerializedActor = ActorData.Find(OriginalObjectPath);
+	if (SerializedActor)
 	{
-		// Subobject of same name already exists and is compatible
-		// TODO: Check whether subobject is marked as serialized and if not serialize using SelectionMap 	
+		return SerializedActor->GetDeserialized(TempActorWorld->GetWorld(), *this, OriginalObjectPath, LocalisationSnapshotPackage);
 	}
-	
-	return ResolvedObject;
-}
 
-UObject* FWorldSnapshotData::ResolveObjectDependencyForClassDefaultObject(int32 ObjectPathIndex)
-{
-	if (ensure(SerializedObjectReferences.IsValidIndex(ObjectPathIndex)))
-	{
-		const FSoftObjectPath& OriginalObjectPath = SerializedObjectReferences[ObjectPathIndex];
-		return ResolveExternalReference(OriginalObjectPath);
-	}
-	return nullptr;
+	UE_LOG(LogLevelSnapshots, Warning, TEXT("No save data found for actor %s"), *OriginalObjectPath.ToString());
+	return {};
 }
 
-int32 FWorldSnapshotData::AddSubobjectDependency(UObject* ReferenceFromOriginalObject)
+TOptional<FObjectSnapshotData*> FWorldSnapshotData::GetSerializedClassDefaults(UClass* Class)
 {
-	const int32 Index = AddObjectDependency(ReferenceFromOriginalObject);
-	const bool bSubobjectAlreadyRegistered = Subobjects.Contains(Index);
-	if (!bSubobjectAlreadyRegistered)
-	{
-		FSubobjectSnapshotData& SubobjectData = Subobjects.Add(Index);
-		UClass* Class = ReferenceFromOriginalObject->GetClass();
-
-		SubobjectData.Class = Class;
-		SubobjectData.OuterIndex = AddObjectDependency(ReferenceFromOriginalObject->GetOuter());
-
-		FTakeWorldObjectSnapshotArchive Serializer = FTakeWorldObjectSnapshotArchive::MakeArchiveForSavingWorldObject(SubobjectData, *this, ReferenceFromOriginalObject);
-		ReferenceFromOriginalObject->Serialize(Serializer);
-
-		AddClassDefault(Class);
-	}
-	
-	return Index;
+	FObjectSnapshotData* ClassDefaultData = ClassDefaults.Find(Class);
+	return ClassDefaultData ? ClassDefaultData : TOptional<FObjectSnapshotData*>();
 }
 
 int32 FWorldSnapshotData::AddCustomSubobjectDependency(UObject* ReferenceFromOriginalObject)
@@ -584,23 +304,13 @@ void FWorldSnapshotData::PostSerialize(const FArchive& Ar)
 {
 	if (Ar.IsLoading() && !SnapshotVersionInfo.IsInitialized())
 	{
-		SnapshotVersionInfo.Initialize();
+		// Assets saved before we added version tracking need to receive versioning info of 4.27.
+		// Skip snapshot version info because it did not exist yet at that time (you'll get migration bugs otherwise).
+		const bool bWithoutSnapshotVersion = true;
+		SnapshotVersionInfo.Initialize(bWithoutSnapshotVersion);
 	}
 }
 
-UObject* FWorldSnapshotData::ResolveExternalReference(const FSoftObjectPath& ObjectPath)
-{
-	UObject* ExternalReference = ObjectPath.ResolveObject();
-	const bool bNeedsToLoadFromDisk = ExternalReference == nullptr;
-	if (bNeedsToLoadFromDisk)
-	{
-		ExternalReference = ObjectPath.TryLoad();
-	}
-		
-	// We're supposed to be dealing with an external (asset) reference, e.g. a UMaterial.
-	ensureAlwaysMsgf(ExternalReference == nullptr || (!ExternalReference->IsA<AActor>() && !ExternalReference->IsA<UActorComponent>()), TEXT("Something is wrong. We just checked that the reference is not a world object but it is an actor or component."));
-	return ExternalReference;
-}
 
 void FWorldSnapshotData::ApplyToWorld_HandleRemovingActors(UWorld* WorldToApplyTo, const FPropertySelectionMap& PropertiesToSerialize)
 {
@@ -635,8 +345,6 @@ void FWorldSnapshotData::ApplyToWorld_HandleRemovingActors(UWorld* WorldToApplyT
 	const bool bVerifyDeletionCanHappen = true;
 	const bool bWarnAboutReferences = false;
 	GEditor->edactDeleteSelected(WorldToApplyTo, bVerifyDeletionCanHappen, bWarnAboutReferences, bWarnAboutReferences);
-	
-	// TODO: restore user selection
 
 	// ... and call the end event like in UUnrealEdEngine
 	FEditorDelegates::OnDeleteActorsEnd.Broadcast();
@@ -714,7 +422,7 @@ void FWorldSnapshotData::ApplyToWorld_HandleRecreatingActors(TSet<AActor*>& Eval
 		{
 			UncastWorld = PathToOwningWorldAsset.TryLoad();
 		}
-		if (!UncastWorld)
+		if (!ensure(UncastWorld))
 		{
 			UE_LOG(LogLevelSnapshots, Error, TEXT("Failed to resolve world '%s'. The world should be loaded?"), *PathToOwningWorldAsset.ToString());
 			continue;
@@ -813,4 +521,4 @@ void FWorldSnapshotData::ApplyToWorld_HandleSerializingMatchingActors(TSet<AActo
 	}
 }
 
-#undef LOCTEXT_NAMESPACE
\ No newline at end of file
+#undef LOCTEXT_NAMESPACE
diff --git a/Engine/Plugins/VirtualProduction/LevelSnapshots/Source/LevelSnapshots/Private/LevelSnapshotsFunctionLibrary.cpp b/Engine/Plugins/VirtualProduction/LevelSnapshots/Source/LevelSnapshots/Private/LevelSnapshotsFunctionLibrary.cpp
--- a/Engine/Plugins/VirtualProduction/LevelSnapshots/Source/LevelSnapshots/Private/LevelSnapshotsFunctionLibrary.cpp
+++ b/Engine/Plugins/VirtualProduction/LevelSnapshots/Source/LevelSnapshots/Private/LevelSnapshotsFunctionLibrary.cpp
@@ -149,7 +149,8 @@ void ULevelSnapshotsFunctionLibrary::ForEachMatchingCustomSubobjectPair(
 	ULevelSnapshot* Snapshot,
 	UObject* SnapshotRootObject,
 	UObject* WorldRootObject,
-	TFunction<void(UObject* SnapshotSubobject, UObject* EditorWorldSubobject)> Callback)
+	TFunction<void(UObject* SnapshotSubobject, UObject* EditorWorldSubobject)> HandleCustomSubobjectPair,
+	TFunction<void(UObject* UnmatchedSnapshotSubobject)> HandleUnmatchedSnapshotSubobject)
 {
-	FCustomObjectSerializationWrapper::ForEachMatchingCustomSubobjectPair(Snapshot->GetSerializedData(), SnapshotRootObject, WorldRootObject, Callback);
+	FCustomObjectSerializationWrapper::ForEachMatchingCustomSubobjectPair(Snapshot->GetSerializedData(), SnapshotRootObject, WorldRootObject, HandleCustomSubobjectPair, HandleUnmatchedSnapshotSubobject);
 }
