Refresh UE5 branch with latest Level Snapshots which includes support for subobjects.

diff --git a/Engine/Plugins/VirtualProduction/LevelSnapshots/Source/LevelSnapshots/Private/LevelSnapshotsModule.cpp b/Engine/Plugins/VirtualProduction/LevelSnapshots/Source/LevelSnapshots/Private/LevelSnapshotsModule.cpp
--- a/Engine/Plugins/VirtualProduction/LevelSnapshots/Source/LevelSnapshots/Private/LevelSnapshotsModule.cpp
+++ b/Engine/Plugins/VirtualProduction/LevelSnapshots/Source/LevelSnapshots/Private/LevelSnapshotsModule.cpp
@@ -4,34 +4,77 @@
 
 #include "BlacklistRestorabilityOverrider.h"
 #include "LevelSnapshotsEditorProjectSettings.h"
+#include "LevelSnapshotsLog.h"
 #include "Restorability/PropertyComparisonParams.h"
 #include "Restorability/StaticMeshCollisionPropertyComparer.h"
 
-#include "EngineUtils.h"
-#include "LevelSnapshotsLog.h"
 #include "Algo/AllOf.h"
+#include "Components/ActorComponent.h"
+#include "EngineUtils.h"
+#include "Algo/Transform.h"
+#include "Engine/Brush.h"
+#include "GameFramework/Actor.h"
+#include "Materials/MaterialInstance.h"
 #include "Modules/ModuleManager.h"
 
 namespace
 {
 	void AddSoftObjectPathSupport(FLevelSnapshotsModule& Module)
 	{
-		// By default FSnapshotRestorability::IsRestorableProperty requires properties to have the CPF_Edit specifier
+		// FSnapshotRestorability::IsRestorableProperty requires properties to have the CPF_Edit specifier
 		// FSoftObjectPath does not have this so we need to whitelist its properties
 
-		UStruct* SoftObjectPath = FindObject<UStruct>(nullptr, TEXT("/Script/CoreUObject.SoftObjectPath"));
-		if (!ensureMsgf(SoftObjectPath, TEXT("Investigate why this class could not be found")))
+		UStruct* SoftObjectClassPath = FindObject<UStruct>(nullptr, TEXT("/Script/CoreUObject.SoftObjectPath"));
+		if (!ensureMsgf(SoftObjectClassPath, TEXT("Investigate why this class could not be found")))
 		{
 			return;
 		}
 
-		TSet<const FProperty*> Properties;
-		for (TFieldIterator<FProperty> FieldIt(SoftObjectPath); FieldIt; ++FieldIt)
+		TSet<const FProperty*> WhitelistedProperties;
+		Algo::Transform(TFieldRange<const FProperty>(SoftObjectClassPath), WhitelistedProperties, [](const FProperty* Prop) { return Prop;} );
+		Module.AddWhitelistedProperties(WhitelistedProperties);
+	}
+
+	void AddAttachParentSupport(FLevelSnapshotsModule& Module)
+	{
+		// These properties are not visible by default because they're not CPF_Edit
+		const FProperty* AttachParent = USceneComponent::StaticClass()->FindPropertyByName(FName("AttachParent"));
+		const FProperty* AttachSocketName = USceneComponent::StaticClass()->FindPropertyByName(FName("AttachSocketName"));
+		if (ensure(AttachParent && AttachSocketName))
 		{
-			Properties.Add(*FieldIt);
+			Module.AddWhitelistedProperties({ AttachParent, AttachSocketName });
 		}
+	}
 
-		Module.AddWhitelistedProperties(Properties);
+	void DisableIrrelevantBrushSubobjects(FLevelSnapshotsModule& Module)
+	{
+		// ABrush::BrushBuilder is CPF_Edit but no user ever cares about it. We don't want it to make volumes to show up as changed.
+		const FProperty* BrushBuilder = ABrush::StaticClass()->FindPropertyByName(GET_MEMBER_NAME_CHECKED(ABrush, BrushBuilder));
+		if (ensure(BrushBuilder))
+		{
+			Module.AddBlacklistedProperties({ BrushBuilder });
+		}
+	}
+
+	void DisableIrrelevantWorldSettings(FLevelSnapshotsModule& Module)
+	{
+		// AWorldSettings::NavigationSystemConfig is CPF_Edit but no user ever cares about it.
+		const FProperty* NavigationSystemConfig = AWorldSettings::StaticClass()->FindPropertyByName(FName("NavigationSystemConfig"));
+		if (ensure(NavigationSystemConfig))
+		{
+			Module.AddBlacklistedProperties({ NavigationSystemConfig });
+		}
+	}
+
+	void DisableIrrelevantMaterialInstanceProperties(FLevelSnapshotsModule& Module)
+	{
+		// This property causes diffs sometimes for unexplained reasons when creating in construction script... does not seem to be important
+		const FProperty* BasePropertyOverrides = UMaterialInstance::StaticClass()->FindPropertyByName(GET_MEMBER_NAME_CHECKED(UMaterialInstance, BasePropertyOverrides));
+		if (ensure(BasePropertyOverrides))
+		{
+			Module.AddBlacklistedProperties({ BasePropertyOverrides });
+		}
+		
 	}
 }
 
@@ -47,6 +90,7 @@ FLevelSnapshotsModule& FLevelSnapshotsModule::GetInternalModuleInstance()
 
 void FLevelSnapshotsModule::StartupModule()
 {
+	// Hook up project settings blacklist
 	const TSharedRef<FBlacklistRestorabilityOverrider> Blacklist = MakeShared<FBlacklistRestorabilityOverrider>(
 		FBlacklistRestorabilityOverrider::FGetBlacklist::CreateLambda([]() -> const FRestorationBlacklist&
 		{
@@ -55,15 +99,24 @@ void FLevelSnapshotsModule::StartupModule()
 		})
 	);
 	RegisterRestorabilityOverrider(Blacklist);
-	
+
+	// Enable / disable troublesome properties
 	AddSoftObjectPathSupport(*this);
+	AddAttachParentSupport(*this);
+	DisableIrrelevantBrushSubobjects(*this);
+	DisableIrrelevantWorldSettings(*this);
+	DisableIrrelevantMaterialInstanceProperties(*this);
 
+	// Interact with special engine features
 	FStaticMeshCollisionPropertyComparer::Register(*this);
 }
 
 void FLevelSnapshotsModule::ShutdownModule()
 {
 	Overrides.Reset();
+	PropertyComparers.Reset();
+	CustomSerializers.Reset();
+	RestorationListeners.Reset();
 }
 
 void FLevelSnapshotsModule::AddCanTakeSnapshotDelegate(FName DelegateName, FCanTakeSnapshot Delegate)
@@ -86,6 +139,31 @@ void FLevelSnapshotsModule::UnregisterRestorabilityOverrider(TSharedRef<ISnapsho
 	Overrides.RemoveSwap(Overrider);
 }
 
+void FLevelSnapshotsModule::AddBlacklistedSubobjectClasses(const TSet<UClass*>& Classes)
+{
+	for (UClass* Class : Classes)
+	{
+		check(Class);
+
+		if (!Class
+			|| !ensureAlwaysMsgf(!Class->IsChildOf(AActor::StaticClass()), TEXT("Invalid function input: Actors can never be subobjects. Check your code."))
+			|| !ensureAlwaysMsgf(!Class->IsChildOf(UActorComponent::StaticClass()), TEXT("Invalid function input: Disallow components using RegisterRestorabilityOverrider and implementing ISnapshotRestorabilityOverrider::IsComponentDesirableForCapture instead.")))
+		{
+			continue;
+		}
+
+		BlacklistedSubobjectClasses.Add(Class);
+	}
+}
+
+void FLevelSnapshotsModule::RemoveBlacklistedSubobjectClasses(const TSet<UClass*>& Classes)
+{
+	for (UClass* Class : Classes)
+	{
+		BlacklistedSubobjectClasses.Remove(Class);
+	}
+}
+
 void FLevelSnapshotsModule::RegisterPropertyComparer(UClass* Class, TSharedRef<IPropertyComparer> Comparer)
 {
 	PropertyComparers.FindOrAdd(Class).AddUnique(Comparer);
@@ -133,6 +211,96 @@ void FLevelSnapshotsModule::UnregisterCustomObjectSerializer(UClass* Class)
 	CustomSerializers.Remove(Class);
 }
 
+void FLevelSnapshotsModule::RegisterRestorationListener(TSharedRef<IRestorationListener> Listener)
+{
+	RestorationListeners.AddUnique(Listener);
+}
+
+void FLevelSnapshotsModule::UnregisterRestorationListener(TSharedRef<IRestorationListener> Listener)
+{
+	RestorationListeners.RemoveSingle(Listener);
+}
+
+void FLevelSnapshotsModule::OnPreApplySnapshotProperties(const FApplySnapshotPropertiesParams& Params)
+{
+	SCOPED_SNAPSHOT_CORE_TRACE(RestorationListeners);
+	
+	for (const TSharedRef<IRestorationListener>& Listener : RestorationListeners)
+	{
+		Listener->PreApplySnapshotProperties(Params);
+	}
+}
+
+void FLevelSnapshotsModule::OnPostApplySnapshotProperties(const FApplySnapshotPropertiesParams& Params)
+{
+	SCOPED_SNAPSHOT_CORE_TRACE(RestorationListeners);
+	
+	for (const TSharedRef<IRestorationListener>& Listener : RestorationListeners)
+	{
+		Listener->PostApplySnapshotProperties(Params);
+	}
+}
+
+void FLevelSnapshotsModule::OnPreApplySnapshotToActor(const FApplySnapshotToActorParams& Params)
+{
+	SCOPED_SNAPSHOT_CORE_TRACE(RestorationListeners);
+	
+	for (const TSharedRef<IRestorationListener>& Listener : RestorationListeners)
+	{
+		Listener->PreApplySnapshotToActor(Params);
+	}
+}
+
+void FLevelSnapshotsModule::OnPostApplySnapshotToActor(const FApplySnapshotToActorParams& Params)
+{
+	SCOPED_SNAPSHOT_CORE_TRACE(RestorationListeners);
+	
+	for (const TSharedRef<IRestorationListener>& Listener : RestorationListeners)
+	{
+		Listener->PostApplySnapshotToActor(Params);
+	}
+}
+
+void FLevelSnapshotsModule::OnPreRecreateComponent(const FPreRecreateComponentParams& Params)
+{
+	SCOPED_SNAPSHOT_CORE_TRACE(RestorationListeners);
+	
+	for (const TSharedRef<IRestorationListener>& Listener : RestorationListeners)
+	{
+		Listener->PreRecreateComponent(Params);
+	}
+}
+
+void FLevelSnapshotsModule::OnPostRecreateComponent(UActorComponent* RecreatedComponent)
+{
+	SCOPED_SNAPSHOT_CORE_TRACE(RestorationListeners);
+	
+	for (const TSharedRef<IRestorationListener>& Listener : RestorationListeners)
+	{
+		Listener->PostRecreateComponent(RecreatedComponent);
+	}
+}
+
+void FLevelSnapshotsModule::OnPreRemoveComponent(UActorComponent* ComponentToRemove)
+{
+	SCOPED_SNAPSHOT_CORE_TRACE(RestorationListeners);
+
+	for (const TSharedRef<IRestorationListener>& Listener : RestorationListeners)
+	{
+		Listener->PreRemoveComponent(ComponentToRemove);
+	}
+}
+
+void FLevelSnapshotsModule::OnPostRemoveComponent(const FPostRemoveComponentParams& Params)
+{
+	SCOPED_SNAPSHOT_CORE_TRACE(RestorationListeners);
+
+	for (const TSharedRef<IRestorationListener>& Listener : RestorationListeners)
+	{
+		Listener->PostRemoveComponent(Params);
+	}
+}
+
 void FLevelSnapshotsModule::AddWhitelistedProperties(const TSet<const FProperty*>& Properties)
 {
 	for (const FProperty* Property : Properties)
@@ -165,6 +333,22 @@ void FLevelSnapshotsModule::RemoveBlacklistedProperties(const TSet<const FProper
 	}
 }
 
+bool FLevelSnapshotsModule::IsSubobjectClassBlacklisted(const UClass* Class) const
+{
+	if (Class->IsChildOf(UActorComponent::StaticClass()))
+	{
+		return false;
+	}
+
+	bool bFoundBlacklistedClass = false;
+	for (const UClass* CurrentClass = Class; !bFoundBlacklistedClass && CurrentClass; CurrentClass = CurrentClass->GetSuperClass())
+	{
+		bFoundBlacklistedClass = BlacklistedSubobjectClasses.Contains(CurrentClass);
+	}
+
+	return bFoundBlacklistedClass;
+}
+
 bool FLevelSnapshotsModule::CanTakeSnapshot(const FPreTakeSnapshotEventData& Event) const
 {
 	return Algo::AllOf(CanTakeSnapshotDelegates, [&Event](const TTuple<FName,FCanTakeSnapshot>& Pair)
diff --git a/Engine/Plugins/VirtualProduction/LevelSnapshots/Source/LevelSnapshots/Private/LevelSnapshotsModule.h b/Engine/Plugins/VirtualProduction/LevelSnapshots/Source/LevelSnapshots/Private/LevelSnapshotsModule.h
--- a/Engine/Plugins/VirtualProduction/LevelSnapshots/Source/LevelSnapshots/Private/LevelSnapshotsModule.h
+++ b/Engine/Plugins/VirtualProduction/LevelSnapshots/Source/LevelSnapshots/Private/LevelSnapshotsModule.h
@@ -27,15 +27,21 @@ public:
 	//~ Begin ILevelSnapshotsModule Interface
 	virtual void RegisterRestorabilityOverrider(TSharedRef<ISnapshotRestorabilityOverrider> Overrider) override;
 	virtual void UnregisterRestorabilityOverrider(TSharedRef<ISnapshotRestorabilityOverrider> Overrider) override;
+	virtual void AddBlacklistedSubobjectClasses(const TSet<UClass*>& Classes) override;
+	virtual void RemoveBlacklistedSubobjectClasses(const TSet<UClass*>& Classes) override;
 	virtual void RegisterPropertyComparer(UClass* Class, TSharedRef<IPropertyComparer> Comparer) override;
 	virtual void UnregisterPropertyComparer(UClass* Class, TSharedRef<IPropertyComparer> Comparer) override;
 	virtual void RegisterCustomObjectSerializer(UClass* Class, TSharedRef<ICustomObjectSnapshotSerializer> CustomSerializer, bool bIncludeBlueprintChildClasses = true);
 	virtual void UnregisterCustomObjectSerializer(UClass* Class) override;
+	virtual void RegisterRestorationListener(TSharedRef<IRestorationListener> Listener) override;
+	virtual void UnregisterRestorationListener(TSharedRef<IRestorationListener> Listener) override;
 	virtual void AddWhitelistedProperties(const TSet<const FProperty*>& Properties) override;
 	virtual void RemoveWhitelistedProperties(const TSet<const FProperty*>& Properties) override;
 	virtual void AddBlacklistedProperties(const TSet<const FProperty*>& Properties) override;
 	virtual void RemoveBlacklistedProperties(const TSet<const FProperty*>& Properties) override;
 	//~ Begin ILevelSnapshotsModule Interface
+
+	bool IsSubobjectClassBlacklisted(const UClass* Class) const;
 	
 	const TArray<TSharedRef<ISnapshotRestorabilityOverrider>>& GetOverrides() const;
 	bool IsPropertyWhitelisted(const FProperty* Property) const;
@@ -51,6 +57,18 @@ public:
 	virtual void RemoveCanTakeSnapshotDelegate(FName DelegateName) override;
 	virtual bool CanTakeSnapshot(const FPreTakeSnapshotEventData& Event) const override;
 
+	void OnPreApplySnapshotProperties(const FApplySnapshotPropertiesParams& Params);
+	void OnPostApplySnapshotProperties(const FApplySnapshotPropertiesParams& Params);
+
+	void OnPreApplySnapshotToActor(const FApplySnapshotToActorParams& Params);
+	void OnPostApplySnapshotToActor(const FApplySnapshotToActorParams& Params);
+	
+	void OnPreRecreateComponent(const FPreRecreateComponentParams& Params);
+	void OnPostRecreateComponent(UActorComponent* RecreatedComponent);
+
+	void OnPreRemoveComponent(UActorComponent* ComponentToRemove);
+	void OnPostRemoveComponent(const FPostRemoveComponentParams& Params);
+	
 private:
 
 	struct FCustomSerializer
@@ -61,10 +79,16 @@ private:
 	
 	/* Allows external modules to override what objects and properties are considered by the snapshot system. */
 	TArray<TSharedRef<ISnapshotRestorabilityOverrider>> Overrides;
+
+	/** Subobject classes we do not capture nor restore */
+	TSet<UClass*> BlacklistedSubobjectClasses;
+	
 	/**/
 	TMap<FSoftClassPath, TArray<TSharedRef<IPropertyComparer>>> PropertyComparers;
 	/**/
 	TMap<FSoftClassPath, FCustomSerializer> CustomSerializers;
+	/**/
+	TArray<TSharedRef<IRestorationListener>> RestorationListeners;
 
 	/* Allows these properties even when the default behaviour would exclude them. */
 	TSet<const FProperty*> WhitelistedProperties;
@@ -72,5 +96,5 @@ private:
 	TSet<const FProperty*> BlacklistedProperties;
 
 	/** Map of named delegates for confirming that a level snapshot is possible. */
-	TMap<FName,FCanTakeSnapshot> CanTakeSnapshotDelegates;
+	TMap<FName, FCanTakeSnapshot> CanTakeSnapshotDelegates;
 };
diff --git a/Engine/Plugins/VirtualProduction/LevelSnapshots/Source/LevelSnapshots/Private/PropertyInfoHelpers.cpp b/Engine/Plugins/VirtualProduction/LevelSnapshots/Source/LevelSnapshots/Private/PropertyInfoHelpers.cpp
--- a/Engine/Plugins/VirtualProduction/LevelSnapshots/Source/LevelSnapshots/Private/PropertyInfoHelpers.cpp
+++ b/Engine/Plugins/VirtualProduction/LevelSnapshots/Source/LevelSnapshots/Private/PropertyInfoHelpers.cpp
@@ -4,6 +4,9 @@
 
 #include "UObject/UnrealType.h"
 
+float FPropertyInfoHelpers::FloatComparisonPrecision = 1e-03f;
+double FPropertyInfoHelpers::DoubleComparisonPrecision = 1e-03;
+
 FProperty* FPropertyInfoHelpers::GetParentProperty(const FProperty* Property)
 {
 	return Property->GetOwner<FProperty>();
@@ -80,6 +83,12 @@ bool FPropertyInfoHelpers::IsPropertySubObject(const FProperty* Property)
 	return bIsComponentSubObject;
 }
 
+void FPropertyInfoHelpers::UpdateDecimalComparisionPrecision(float FloatPrecision, double DoublePrecision)
+{
+	FloatComparisonPrecision = FloatPrecision;
+	DoubleComparisonPrecision = DoublePrecision;
+}
+
 bool FPropertyInfoHelpers::AreNumericPropertiesNearlyEqual(const FNumericProperty* NumericProperty, const void* ValuePtrA, const void* ValuePtrB)
 {
 	check(NumericProperty);
@@ -88,18 +97,18 @@ bool FPropertyInfoHelpers::AreNumericPropertiesNearlyEqual(const FNumericPropert
 	{
 		const float ValueA = FloatProperty->GetFloatingPointPropertyValue(ValuePtrA); 
 		const float ValueB = FloatProperty->GetFloatingPointPropertyValue(ValuePtrB);
-		return FMath::IsNearlyEqual(ValueA, ValueB, KINDA_SMALL_NUMBER);
+		return FMath::IsNearlyEqual(ValueA, ValueB, FloatComparisonPrecision);
 	}
 	
 	if (const FDoubleProperty* DoubleProperty = CastField<FDoubleProperty>(NumericProperty))
 	{
 		const double ValueA = DoubleProperty->GetFloatingPointPropertyValue(ValuePtrA);
 		const double ValueB = DoubleProperty->GetFloatingPointPropertyValue(ValuePtrB);
-		return FMath::IsNearlyEqual(ValueA, ValueB, static_cast<double>(KINDA_SMALL_NUMBER));
+		return FMath::IsNearlyEqual(ValueA, ValueB, DoubleComparisonPrecision);
 	}
 	
 	// Not a float or double? Then some kind of integer (byte, int8, int16, int32, int64, uint8, uint16 ...). Enums are bytes.
 	const int64 ValueA = NumericProperty->GetSignedIntPropertyValue(ValuePtrA);
 	const int64 ValueB = NumericProperty->GetSignedIntPropertyValue(ValuePtrB);
 	return ValueA == ValueB;
-}
\ No newline at end of file
+}
diff --git a/Engine/Plugins/VirtualProduction/LevelSnapshots/Source/LevelSnapshots/Private/PropertySelectionMap.cpp b/Engine/Plugins/VirtualProduction/LevelSnapshots/Source/LevelSnapshots/Private/PropertySelectionMap.cpp
--- a/Engine/Plugins/VirtualProduction/LevelSnapshots/Source/LevelSnapshots/Private/PropertySelectionMap.cpp
+++ b/Engine/Plugins/VirtualProduction/LevelSnapshots/Source/LevelSnapshots/Private/PropertySelectionMap.cpp
@@ -2,7 +2,9 @@
 
 #include "PropertySelectionMap.h"
 
+#include "CustomSubobjectRestorationInfo.h"
 #include "LevelSnapshotsLog.h"
+#include "Data/RestorableObjectSelection.h"
 
 #include "GameFramework/Actor.h"
 #include "UObject/UObjectHash.h"
@@ -29,39 +31,151 @@ void FPropertySelectionMap::RemoveNewActorToDespawn(AActor* WorldActor)
 
 bool FPropertySelectionMap::AddObjectProperties(UObject* WorldObject, const FPropertySelection& SelectedProperties)
 {
-	if (!SelectedProperties.IsEmpty() && ensure(WorldObject))
-	{
-		SelectedWorldObjectsToSelectedProperties.FindOrAdd(WorldObject) = SelectedProperties;
+	if (ensure(WorldObject)
+		// Expliticly warn about empty sets for subobjects. See documentation of MarkSubobjectForRestoringReferencesButSkipProperties.
+		&& ensureMsgf(!SelectedProperties.IsEmpty(), TEXT("Maybe you meant to call MarkSubobjectForRestoringReferencesButSkipProperties?"))
+		&& ensureMsgf(!SubobjectsMarkedForReferencesRestorationOnly.Contains(WorldObject), TEXT("Object was already added via MarkSubobjectForRestoringReferencesButSkipProperties because it had no changed properties. Now we're told that is does have selected properties...")))
+	{		
+		EditorWorldObjectToSelectedProperties.FindOrAdd(WorldObject) = SelectedProperties;
 		return true;
 	}
 	return false;
 }
 
+void FPropertySelectionMap::MarkSubobjectForRestoringReferencesButSkipProperties(UObject* WorldSubbject)
+{
+	if (ensure(WorldSubbject)
+		&& ensureMsgf(!EditorWorldObjectToSelectedProperties.Contains(WorldSubbject), TEXT("Object was already added via AddObjectProperties with changed properties. Now we're told that is has no changed properties..."))
+		&& ensureMsgf(!WorldSubbject->IsA<UActorComponent>(), TEXT("Components have special case handling and should not be added.")))
+	{
+		SubobjectsMarkedForReferencesRestorationOnly.Add(WorldSubbject);
+	}
+}
+
+bool FPropertySelectionMap::IsSubobjectMarkedForReferenceRestorationOnly(const FSoftObjectPath& WorldSubobjectPath) const
+{
+	return SubobjectsMarkedForReferencesRestorationOnly.Contains(WorldSubobjectPath);
+}
+
 void FPropertySelectionMap::RemoveObjectPropertiesFromMap(UObject* WorldObject)
 {
-	SelectedWorldObjectsToSelectedProperties.Remove(WorldObject);
+	EditorWorldObjectToSelectedProperties.Remove(WorldObject);
 }
 
-const FPropertySelection* FPropertySelectionMap::GetSelectedProperties(UObject* WorldObject) const
+void FPropertySelectionMap::AddComponentSelection(AActor* EditorWorldActor, const FAddedAndRemovedComponentInfo& ComponentSelection)
 {
-	return GetSelectedProperties(FSoftObjectPath(WorldObject));
+	const bool bIsEmpty = ComponentSelection.SnapshotComponentsToAdd.Num() == 0 && ComponentSelection.EditorWorldComponentsToRemove.Num() == 0;
+	if (ensure(EditorWorldActor && !bIsEmpty))
+	{
+		EditorActorToComponentSelection.Add(EditorWorldActor, ComponentSelection);
+	}
+}
+
+void FPropertySelectionMap::RemoveComponentSelection(AActor* EditorWorldActor)
+{
+	EditorActorToComponentSelection.Remove(EditorWorldActor);
+}
+
+void FPropertySelectionMap::AddCustomEditorSubobjectToRecreate(UObject* EditorWorldOwner, UObject* SnapshotSubobject)
+{
+	EditorWorldObjectToCustomSubobjectSelection.FindOrAdd(EditorWorldOwner).CustomSnapshotSubobjectsToRestore.Add(SnapshotSubobject);
+}
+
+void FPropertySelectionMap::RemoveCustomEditorSubobjectToRecreate(UObject* EditorWorldOwner, UObject* SnapshotSubobject)
+{
+	if (FCustomSubobjectRestorationInfo* RestorationInfo = EditorWorldObjectToCustomSubobjectSelection.Find(EditorWorldOwner))
+	{
+		if (RestorationInfo->CustomSnapshotSubobjectsToRestore.Num() == 1)
+		{
+			EditorWorldObjectToCustomSubobjectSelection.Remove(EditorWorldOwner);
+		}
+		else
+		{
+			RestorationInfo->CustomSnapshotSubobjectsToRestore.Remove(SnapshotSubobject);
+		}
+	}
 }
 
 const FPropertySelection* FPropertySelectionMap::GetSelectedProperties(const FSoftObjectPath& WorldObjectPath) const
 {
-	return SelectedWorldObjectsToSelectedProperties.Find(WorldObjectPath); 
+	return EditorWorldObjectToSelectedProperties.Find(WorldObjectPath); 
+}
+
+FRestorableObjectSelection FPropertySelectionMap::GetObjectSelection(const FSoftObjectPath& EditorWorldObject) const
+{
+	return FRestorableObjectSelection(EditorWorldObject, *this);
 }
 
 TArray<FSoftObjectPath> FPropertySelectionMap::GetKeys() const
 {
 	TArray<FSoftObjectPath> Result;
-	SelectedWorldObjectsToSelectedProperties.GenerateKeyArray(Result);
+	EditorWorldObjectToSelectedProperties.GenerateKeyArray(Result);
+
+	TArray<FSoftObjectPath> ComponentPaths;
+	EditorActorToComponentSelection.GenerateKeyArray(ComponentPaths);
+	LIKELY(ComponentPaths.Num() < Result.Num());
+	for (const FSoftObjectPath& Path : ComponentPaths)
+	{
+		Result.Add(Path);
+	}
+	
 	return Result;
 }
 
+bool FPropertySelectionMap::HasChanges(AActor* EditorWorldActor) const
+{
+	const FRestorableObjectSelection Selection = GetObjectSelection(EditorWorldActor);
+	
+	const FPropertySelection* PropertySelection = Selection.GetPropertySelection();
+	const bool bHasSelectedProperties = PropertySelection && !PropertySelection->IsEmpty();
+	
+	const FAddedAndRemovedComponentInfo* ComponentSelection = Selection.GetComponentSelection();
+	const bool bHasAddedOrRemovedComps = ComponentSelection
+		&& (Selection.GetComponentSelection()->SnapshotComponentsToAdd.Num() != 0 || Selection.GetComponentSelection()->EditorWorldComponentsToRemove.Num() != 0);
+	
+	const FCustomSubobjectRestorationInfo* CustomSubobjectInfo = Selection.GetCustomSubobjectSelection();
+	const bool bHasCustomSubobjects = CustomSubobjectInfo != nullptr;
+	if (bHasSelectedProperties || bHasAddedOrRemovedComps || bHasCustomSubobjects)
+	{
+		return true;
+	}
+		
+	for (UActorComponent* Component : TInlineComponentArray<UActorComponent*>(EditorWorldActor))
+	{
+		if (HasChanges(Component))
+		{
+			return true;
+		}
+	}
+
+	return false;
+}
+
+bool FPropertySelectionMap::HasChanges(UActorComponent* EditorWorldComponent) const
+{
+	const FRestorableObjectSelection Selection = GetObjectSelection(EditorWorldComponent);
+	const FPropertySelection* PropertySelection = Selection.GetPropertySelection();
+	const FCustomSubobjectRestorationInfo* CustomSubobjectInfo = Selection.GetCustomSubobjectSelection();
+	return (PropertySelection && !PropertySelection->IsEmpty()) || CustomSubobjectInfo;
+}
+
+bool FPropertySelectionMap::HasChanges(UObject* EditorWorldObject) const
+{
+	if (AActor* Actor = Cast<AActor>(EditorWorldObject))
+	{
+		return HasChanges(Actor);
+	}
+	
+	const FRestorableObjectSelection Selection = GetObjectSelection(EditorWorldObject);
+	const FPropertySelection* PropertySelection = Selection.GetPropertySelection();
+	const FCustomSubobjectRestorationInfo* CustomSubobjectInfo = Selection.GetCustomSubobjectSelection();
+	return (PropertySelection && !PropertySelection->IsEmpty()) || CustomSubobjectInfo;
+}
+
 void FPropertySelectionMap::Empty(bool bCanShrink)
 {
-	SelectedWorldObjectsToSelectedProperties.Empty(bCanShrink ? SelectedWorldObjectsToSelectedProperties.Num() : 0);
+	EditorWorldObjectToSelectedProperties.Empty(bCanShrink ? EditorWorldObjectToSelectedProperties.Num() : 0);
+	EditorActorToComponentSelection.Empty(bCanShrink ? EditorActorToComponentSelection.Num() : 0);
 	DeletedActorsToRespawn.Empty(bCanShrink ? DeletedActorsToRespawn.Num() : 0);
 	NewActorsToDespawn.Empty(bCanShrink ? NewActorsToDespawn.Num() : 0);
 }
@@ -75,7 +189,7 @@ TArray<UObject*> FPropertySelectionMap::GetDirectSubobjectsWithProperties(UObjec
 
 	for (int32 i = Subobjects.Num() - 1; i > 0; --i)
 	{
-		const bool bHasSelectedProperties = GetSelectedProperties(Subobjects[i]) != nullptr;
+		const bool bHasSelectedProperties = GetObjectSelection(Subobjects[i]).GetPropertySelection() != nullptr;
 		if (!bHasSelectedProperties)
 		{
 			Subobjects.RemoveAt(i);
diff --git a/Engine/Plugins/VirtualProduction/LevelSnapshots/Source/LevelSnapshots/Private/Restorability/SnapshotRestorability.cpp b/Engine/Plugins/VirtualProduction/LevelSnapshots/Source/LevelSnapshots/Private/Restorability/SnapshotRestorability.cpp
--- a/Engine/Plugins/VirtualProduction/LevelSnapshots/Source/LevelSnapshots/Private/Restorability/SnapshotRestorability.cpp
+++ b/Engine/Plugins/VirtualProduction/LevelSnapshots/Source/LevelSnapshots/Private/Restorability/SnapshotRestorability.cpp
@@ -2,11 +2,16 @@
 
 #include "Restorability/SnapshotRestorability.h"
 
+#include "Landscape.h"
+#include "LandscapeGizmoActor.h"
 #include "LevelSnapshotsLog.h"
 #include "LevelSnapshotsModule.h"
+#include "AI/NavigationSystemConfig.h"
+#include "Algo/AllOf.h"
 
 #include "Components/ActorComponent.h"
 #include "Components/BillboardComponent.h"
+#include "Engine/BrushBuilder.h"
 #include "Engine/LevelScriptActor.h"
 #include "GameFramework/Actor.h"
 #include "GameFramework/DefaultPhysicsVolume.h"
@@ -23,36 +28,28 @@ namespace
 		const TSet<UClass*> UnsupportedClasses = 
 		{
 			ALevelScriptActor::StaticClass(),		// The level blueprint. Filtered out to avoid external map errors when saving a snapshot.
-            ADefaultPhysicsVolume::StaticClass()	// Does not show up in world outliner; always spawned with world.
+            ADefaultPhysicsVolume::StaticClass(),	// Does not show up in world outliner; always spawned with world.
+
+			// Until proper landscape support is implemented, let's completely ignore landscapes to avoid capturing incomplete data
+			ALandscapeProxy::StaticClass(),
+			ALandscapeGizmoActor::StaticClass()
         };
 
-		for (UClass* Class : UnsupportedClasses)
-		{
-			if (Actor->IsA(Class))
-			{
-				return false;
-			}
-		}
-	
-		return true;
+		return Algo::AllOf(UnsupportedClasses, [Actor](UClass *Class) { return !Actor->IsA(Class); });
 	}
 
 	bool DoesActorHaveSupportedClassForRemoving(const AActor* Actor)
 	{
 		const TSet<UClass*> UnsupportedClasses = 
 		{
-			AWorldSettings::StaticClass()			// Every sublevel creates a world settings. We do not want to be respawning them
-		};
+			AWorldSettings::StaticClass(),			// Every sublevel creates a world settings. We do not want to be respawning them
 
-		for (UClass* Class : UnsupportedClasses)
-		{
-			if (Actor->IsA(Class))
-			{
-				return false;
-			}
-		}
-	
-		return true;
+			// Until proper landscape support is implemented, let's completely ignore landscapes to avoid capturing incomplete data
+			ALandscapeProxy::StaticClass(),
+			ALandscapeGizmoActor::StaticClass()
+		};
+		
+		return Algo::AllOf(UnsupportedClasses, [Actor](UClass *Class) { return !Actor->IsA(Class); });
 	}
 
 	bool DoesComponentHaveSupportedClassForCapture(const UActorComponent* Component)
@@ -61,22 +58,30 @@ namespace
 		{
 			UBillboardComponent::StaticClass(),		// Attached to all editor actors > It always has a different name so we will never be able to match it.
         };
+		
+		return Algo::AllOf(UnsupportedClasses, [Component](UClass *Class) { return !Component->IsA(Class); });
+	}
 
-		for (UClass* Class : UnsupportedClasses)
+	bool DoesSubobjecttHaveSupportedClassForCapture(const UObject* Subobject)
+	{
+		const TSet<UClass*> UnsupportedClasses = 
 		{
-			if (Component->IsA(Class))
-			{
-				return false;
-			}
-		}
-	
-		return true;
+			UBrushBuilder::StaticClass(),				// Does not have state. Referenced by AVolumes, ABrushes, e.g. APostProcessVolume or ALightmassImportanceVolume
+			UNavigationSystemConfig::StaticClass()		
+		};
+		
+		return Algo::AllOf(UnsupportedClasses, [Subobject](UClass *Class) { return !Subobject->IsA(Class); });
 	}
 }
 
 bool FSnapshotRestorability::IsActorDesirableForCapture(const AActor* Actor)
 {
 	SCOPED_SNAPSHOT_CORE_TRACE(IsActorDesirableForCapture);
+
+	if (!IsValid(Actor))
+	{
+		return false;
+	}
 	
 	bool bSomebodyAllowed = false;
 	const TArray<TSharedRef<ISnapshotRestorabilityOverrider>>& Overrides = FLevelSnapshotsModule::GetInternalModuleInstance().GetOverrides(); 
@@ -95,11 +100,8 @@ bool FSnapshotRestorability::IsActorDesirableForCapture(const AActor* Actor)
 		return true;
 	}
 
-
-
 	
-	return IsValid(Actor)
-            && DoesActorHaveSupportedClass(Actor)
+	return DoesActorHaveSupportedClass(Actor)
             && !Actor->IsTemplate()								// Should never happen, but we never want CDOs
             && !Actor->HasAnyFlags(RF_Transient)				// Don't add transient actors in non-play worlds		
 #if WITH_EDITOR
@@ -113,6 +115,11 @@ bool FSnapshotRestorability::IsActorDesirableForCapture(const AActor* Actor)
 bool FSnapshotRestorability::IsComponentDesirableForCapture(const UActorComponent* Component)
 {
 	SCOPED_SNAPSHOT_CORE_TRACE(IsComponentDesirableForCapture);
+
+	if (!IsValid(Component))
+	{
+		return false;
+	}
 	
 	bool bSomebodyAllowed = false;
 	const TArray<TSharedRef<ISnapshotRestorabilityOverrider>>& Overrides = FLevelSnapshotsModule::GetInternalModuleInstance().GetOverrides(); 
@@ -125,17 +132,48 @@ bool FSnapshotRestorability::IsComponentDesirableForCapture(const UActorComponen
 			return false;
 		}
 	}
+	
+	return bSomebodyAllowed
+		// Components created in construction script are not supported
+		|| (Component->CreationMethod != EComponentCreationMethod::UserConstructionScript
+			&& DoesComponentHaveSupportedClassForCapture(Component));
+}
 
-	if (bSomebodyAllowed)
+bool FSnapshotRestorability::IsSubobjectClassDesirableForCapture(const UClass* SubobjectClass)
+{
+	SCOPED_SNAPSHOT_CORE_TRACE(IsSubobjectClassDesirableForCapture);
+	return DoesSubobjecttHaveSupportedClassForCapture(SubobjectClass) && !FLevelSnapshotsModule::GetInternalModuleInstance().IsSubobjectClassBlacklisted(SubobjectClass);
+}
+
+bool FSnapshotRestorability::IsSubobjectDesirableForCapture(const UObject* Subobject)
+{
+	checkf(!Subobject->IsA<UClass>(), TEXT("Do you have a typo on your code?"));
+	
+	if (const UActorComponent* Component = Cast<UActorComponent>(Subobject))
 	{
-		return true;
+		return IsComponentDesirableForCapture(Component);	
 	}
-	
 
+	return Subobject ? IsSubobjectClassDesirableForCapture(Subobject->GetClass()) : false;
+}
+
+bool FSnapshotRestorability::IsPropertyDesirableForCapture(const FProperty* Property)
+{
+	SCOPED_SNAPSHOT_CORE_TRACE(IsPropertyDesirableForCapture);
 	
-	// We only support native components or the ones added through the component list in Blueprints for now
-	return IsValid(Component) && DoesComponentHaveSupportedClassForCapture(Component)
-		&& (Component->CreationMethod == EComponentCreationMethod::Native || Component->CreationMethod == EComponentCreationMethod::SimpleConstructionScript);
+	const bool bIsWhitelisted = IsPropertyWhitelistedForCapture(Property);
+	const bool bIsBlacklisted = IsPropertyBlacklistedForCapture(Property);
+
+	// To avoid saving every single subobject, only save the editable ones.
+	const uint64 EditablePropertyFlag = CPF_Edit;
+	const bool bIsObjectProperty = CastField<FObjectPropertyBase>(Property) != nullptr;
+
+	// TODO: Technically, we should only save editable/visible properties. However, we not know whether this will cause any trouble for existing snapshots or has unintended side-effects.
+	const bool bIsAllowed = !bIsBlacklisted
+		&& (bIsWhitelisted
+			|| !bIsObjectProperty															// Legacy / safety reasons: allow all non-object properties
+			|| (bIsObjectProperty && Property->HasAnyPropertyFlags(EditablePropertyFlag))); // Implementing object support now... rather save less than too many subobjects
+	return bIsAllowed;
 }
 
 bool FSnapshotRestorability::IsPropertyBlacklistedForCapture(const FProperty* Property)
@@ -159,9 +197,6 @@ bool FSnapshotRestorability::IsRestorableProperty(const FProperty* LeafProperty)
 {
 	SCOPED_SNAPSHOT_CORE_TRACE(IsRestorableProperty);
 	
-	// Subobjects are currently not supported
-	const uint64 InstancedFlags = CPF_InstancedReference | CPF_ContainsInstancedReference | CPF_PersistentInstance;
-	
 	// Deprecated and transient properties should not cause us to consider the property different because we do not save these properties.
 	const uint64 UnsavedProperties = CPF_Deprecated | CPF_Transient;
 	// Property is not editable in details panels
@@ -175,7 +210,7 @@ bool FSnapshotRestorability::IsRestorableProperty(const FProperty* LeafProperty)
 	const bool bIsWhitelisted = Module.IsPropertyWhitelisted(LeafProperty);
 	const bool bIsBlacklisted = Module.IsPropertyBlacklisted(LeafProperty);
 	const bool bPassesDefaultChecks =
-		!LeafProperty->HasAnyPropertyFlags(UnsavedProperties | InstancedFlags | UneditableFlags)
+		!LeafProperty->HasAnyPropertyFlags(UnsavedProperties | UneditableFlags)
         && LeafProperty->HasAllPropertyFlags(RequiredFlags);
 	
 	return bIsWhitelisted || (!bIsBlacklisted && bPassesDefaultChecks);
diff --git a/Engine/Plugins/VirtualProduction/LevelSnapshots/Source/LevelSnapshots/Private/Settings/LevelSnapshotsEditorProjectSettings.cpp b/Engine/Plugins/VirtualProduction/LevelSnapshots/Source/LevelSnapshots/Private/Settings/LevelSnapshotsEditorProjectSettings.cpp
--- a/Engine/Plugins/VirtualProduction/LevelSnapshots/Source/LevelSnapshots/Private/Settings/LevelSnapshotsEditorProjectSettings.cpp
+++ b/Engine/Plugins/VirtualProduction/LevelSnapshots/Source/LevelSnapshots/Private/Settings/LevelSnapshotsEditorProjectSettings.cpp
@@ -2,6 +2,7 @@
 
 #include "Settings/LevelSnapshotsEditorProjectSettings.h"
 
+#include "PropertyInfoHelpers.h"
 #include "Application/SlateApplicationBase.h"
 #include "HAL/PlatformApplicationMisc.h"
 
@@ -34,3 +35,20 @@ void ULevelSnapshotsEditorProjectSettings::SetLastCreationWindowSize(const FVect
 	PreferredCreationFormWindowWidth = InLastSize.X;
 	PreferredCreationFormWindowHeight = InLastSize.Y;
 }
+
+void ULevelSnapshotsEditorProjectSettings::PostInitProperties()
+{
+	UObject::PostInitProperties();
+	FPropertyInfoHelpers::UpdateDecimalComparisionPrecision(FloatComparisonPrecision, DoubleComparisonPrecision);
+}
+
+void ULevelSnapshotsEditorProjectSettings::PostEditChangeProperty(FPropertyChangedEvent& PropertyChangedEvent)
+{
+	if (PropertyChangedEvent.GetPropertyName() == GET_MEMBER_NAME_CHECKED(ULevelSnapshotsEditorProjectSettings, FloatComparisonPrecision)
+		|| PropertyChangedEvent.GetPropertyName() == GET_MEMBER_NAME_CHECKED(ULevelSnapshotsEditorProjectSettings, DoubleComparisonPrecision))
+	{
+		FPropertyInfoHelpers::UpdateDecimalComparisionPrecision(FloatComparisonPrecision, DoubleComparisonPrecision);
+	}
+	
+	UObject::PostEditChangeProperty(PropertyChangedEvent);
+}
diff --git a/Engine/Plugins/VirtualProduction/LevelSnapshots/Source/LevelSnapshots/Private/Tests/ApplyToWorldTests.cpp b/Engine/Plugins/VirtualProduction/LevelSnapshots/Source/LevelSnapshots/Private/Tests/ApplyToWorldTests.cpp
--- a/Engine/Plugins/VirtualProduction/LevelSnapshots/Source/LevelSnapshots/Private/Tests/ApplyToWorldTests.cpp
+++ b/Engine/Plugins/VirtualProduction/LevelSnapshots/Source/LevelSnapshots/Private/Tests/ApplyToWorldTests.cpp
@@ -1,5 +1,6 @@
 ﻿// Copyright Epic Games, Inc. All Rights Reserved.
 
+#include "PropertySelectionMap.h"
 #include "SnapshotTestRunner.h"
 #include "SnapshotTestActor.h"
 
@@ -9,6 +10,7 @@
 #include "Engine/StaticMeshActor.h"
 #include "Materials/MaterialInterface.h"
 #include "Misc/AutomationTest.h"
+#include "Util/EquivalenceUtil.h"
 
 IMPLEMENT_SIMPLE_AUTOMATION_TEST(FRestoreSimpleProperties, "VirtualProduction.LevelSnapshots.Snapshot.RestoreSimpleProperties", (EAutomationTestFlags::ApplicationContextMask | EAutomationTestFlags::EngineFilter));
 bool FRestoreSimpleProperties::RunTest(const FString& Parameters)
@@ -138,6 +140,17 @@ bool FRestoreReferenceProperties::RunTest(const FString& Parameters)
 	return true;
 }
 
+/**
+ * Things like changing array / set / map element order, such as { A, null } to { null, A }, is detected and restored correctly.
+ * Differs from FReorderSubobjectCollections because different code runs for subobjects.
+ */
+IMPLEMENT_SIMPLE_AUTOMATION_TEST(FReorderReferenceCollections, "VirtualProduction.LevelSnapshots.Snapshot.ReorderReferenceCollections", (EAutomationTestFlags::ApplicationContextMask | EAutomationTestFlags::EngineFilter));
+bool FReorderReferenceCollections::RunTest(const FString& Parameters)
+{
+	// TODO: 
+	return true;
+}
+
 IMPLEMENT_SIMPLE_AUTOMATION_TEST(FInstancedStaticMesh, "VirtualProduction.LevelSnapshots.Snapshot.InstancedStaticMesh", (EAutomationTestFlags::ApplicationContextMask | EAutomationTestFlags::EngineFilter));
 bool FInstancedStaticMesh::RunTest(const FString& Parameters)
 {
@@ -209,3 +222,88 @@ bool FSkipTransientAndDeprecatedProperties::RunTest(const FString& Parameters)
 	return true;
 }
 
+/**
+ *
+ */
+IMPLEMENT_SIMPLE_AUTOMATION_TEST(FUnchangedActorHasNoDiff, "VirtualProduction.LevelSnapshots.Snapshot.UnchangedActorHasNoDiff", (EAutomationTestFlags::ApplicationContextMask | EAutomationTestFlags::EngineFilter));
+bool FUnchangedActorHasNoDiff::RunTest(const FString& Parameters)
+{
+	ASnapshotTestActor* Actor = nullptr;
+	FSnapshotTestRunner()
+		.ModifyWorld([&](UWorld* World)
+		{
+			Actor = World->SpawnActor<ASnapshotTestActor>();
+			Actor->AllocateSubobjects();
+		})
+		.TakeSnapshot()
+		.FilterProperties(Actor, [&](const FPropertySelectionMap& PropertySelection)
+		{
+			TestEqual(TEXT("Actor stays unchanged"), PropertySelection.GetKeyCount(), 0);
+		});
+
+	return true;
+}
+
+IMPLEMENT_SIMPLE_AUTOMATION_TEST(FActorAttachParentRestores, "VirtualProduction.LevelSnapshots.Snapshot.ActorAttachParentRestores", (EAutomationTestFlags::ApplicationContextMask | EAutomationTestFlags::EngineFilter));
+bool FActorAttachParentRestores::RunTest(const FString& Parameters)
+{
+	ASnapshotTestActor* ParentOne = nullptr;
+	ASnapshotTestActor* ParentTwo = nullptr;
+	ASnapshotTestActor* Child = nullptr;
+	
+	FSnapshotTestRunner()
+		.ModifyWorld([&](UWorld* World)
+		{
+			ParentOne = World->SpawnActor<ASnapshotTestActor>();
+			ParentTwo = World->SpawnActor<ASnapshotTestActor>();
+			Child = World->SpawnActor<ASnapshotTestActor>();
+
+			Child->AttachToActor(ParentOne, FAttachmentTransformRules::KeepRelativeTransform);
+		})
+		.TakeSnapshot()
+
+		.ModifyWorld([&](UWorld* World)
+		{
+			Child->AttachToActor(ParentTwo, FAttachmentTransformRules::KeepRelativeTransform);
+		})
+		.FilterProperties(Child, [&](const FPropertySelectionMap& PropertySelection)
+		{
+			const FPropertySelection* ComponentSelection = PropertySelection.GetObjectSelection(Child->GetRootComponent()).GetPropertySelection();
+
+			const bool bRootComponentSelected = ComponentSelection != nullptr;
+			TestTrue(TEXT("Root Component has changed properties"), bRootComponentSelected);
+			TestEqual(TEXT("No other objects changed"), PropertySelection.GetKeyCount(), 1);
+			if (bRootComponentSelected)
+			{
+				TestEqual(TEXT("No other properties changed"), ComponentSelection->GetSelectedProperties().Num(), 1);
+
+				const FProperty* AttachParentProperty = USceneComponent::StaticClass()->FindPropertyByName(FName("AttachParent"));
+				ensureMsgf(AttachParentProperty, TEXT("Did property name change?"));
+				TestTrue(TEXT("AttachParent changed"), ComponentSelection->IsPropertySelected(nullptr, AttachParentProperty));
+			}
+		})
+
+		.AccessSnapshot([&](ULevelSnapshot* Snapshot)
+		{
+			TOptional<AActor*> SnapshotChild = Snapshot->GetDeserializedActor(Child);
+			TOptional<AActor*> SnapshotParentOne = Snapshot->GetDeserializedActor(ParentOne);
+
+			const bool bParentsWereDeserialized = SnapshotChild && SnapshotParentOne;
+			TestTrue(TEXT("Parents were deserialized"), bParentsWereDeserialized);
+			if (bParentsWereDeserialized)
+			{
+				const bool bHasOriginalChanged = SnapshotUtil::HasOriginalChangedPropertiesSinceSnapshotWasTaken(Snapshot->GetSerializedData(), *SnapshotChild, Child);
+				TestTrue(TEXT("HasOriginalChangedPropertiesSinceSnapshotWasTaken == true"), bHasOriginalChanged);
+				
+				TestTrue(TEXT("Snapshot has correct attach parent"), SnapshotChild.GetValue()->GetAttachParentActor() == SnapshotParentOne);
+			}
+		})
+
+		.ApplySnapshot()
+		.RunTest([&]()
+		{
+			TestTrue(TEXT("Attach parent restored correctly"), Child->GetAttachParentActor() == ParentOne);
+		});
+
+	return true;
+}
\ No newline at end of file
diff --git a/Engine/Plugins/VirtualProduction/LevelSnapshots/Source/LevelSnapshots/Private/Tests/ComponentTests.cpp b/Engine/Plugins/VirtualProduction/LevelSnapshots/Source/LevelSnapshots/Private/Tests/ComponentTests.cpp
new file mode 100644
--- /dev/null
+++ b/Engine/Plugins/VirtualProduction/LevelSnapshots/Source/LevelSnapshots/Private/Tests/ComponentTests.cpp
@@ -0,0 +1,383 @@
+﻿// Copyright Epic Games, Inc. All Rights Reserved.
+
+#include "EngineUtils.h"
+#include "PropertySelection.h"
+#include "PropertySelectionMap.h"
+#include "SnapshotTestRunner.h"
+#include "SnapshotTestActor.h"
+
+#include "Components/PointLightComponent.h"
+#include "Components/StaticMeshComponent.h"
+#include "Engine/StaticMesh.h"
+#include "Engine/World.h"
+#include "GameFramework/Actor.h"
+#include "Misc/AutomationTest.h"
+
+// TODO: FExistingSubobjectHasDifferentClass > Change a subobject's class and try to restore into it
+// TODO: FAttachementRestoration
+
+/**
+ * Checks that instanced and native components are diffed correctly:
+ *  - Added components are marked to be removed
+ *  - Removed components are marked to be added
+ *
+ *  Additionally, this indirectly tests that components are recreated correctly on the snapshot actor:
+ *  if components were not recreated correctly on the snapshot actor, the diff would not generate correctly in the first place.
+ */
+IMPLEMENT_SIMPLE_AUTOMATION_TEST(FAddedAndRemovedComponentsCorrectOnSnapshotActorAndGenerateDiff, "VirtualProduction.LevelSnapshots.Snapshot.Component.AddedAndRemovedComponentsCorrectOnSnapshotActorAndGenerateDiff", (EAutomationTestFlags::ApplicationContextMask | EAutomationTestFlags::EngineFilter));
+bool FAddedAndRemovedComponentsCorrectOnSnapshotActorAndGenerateDiff::RunTest(const FString& Parameters)
+{
+	AActor* Actor = nullptr;
+	UStaticMeshComponent* InstancedStaysOnActor = nullptr;
+	UStaticMeshComponent* InstancedRemovedFromActor = nullptr;
+	UStaticMeshComponent* InstancedAddedToActor = nullptr;
+	UStaticMeshComponent* NativeStaysOnActor = nullptr;
+	UStaticMeshComponent* NativeRemovedFromActor = nullptr;
+	UStaticMeshComponent* NativeAddedToActor = nullptr;
+
+	// Generate properties
+	FSnapshotTestRunner()
+		.ModifyWorld([&](UWorld* World)
+		{
+			Actor = World->SpawnActor<AActor>();
+
+			InstancedStaysOnActor = NewObject<UStaticMeshComponent>(Actor, UStaticMeshComponent::StaticClass(), "InstancedStaysOnActor");
+			Actor->AddInstanceComponent(InstancedStaysOnActor);
+			InstancedRemovedFromActor = NewObject<UStaticMeshComponent>(Actor, UStaticMeshComponent::StaticClass(), "InstancedRemovedFromActor");
+			Actor->AddInstanceComponent(InstancedRemovedFromActor);
+			NativeStaysOnActor = NewObject<UStaticMeshComponent>(Actor, UStaticMeshComponent::StaticClass(), "NativeStaysOnActor");
+			NativeRemovedFromActor = NewObject<UStaticMeshComponent>(Actor, UStaticMeshComponent::StaticClass(), "NativeRemovedFromActor");
+
+			InstancedStaysOnActor->SetRelativeLocation(FVector(10,20,30));
+			NativeStaysOnActor->SetRelativeLocation(FVector(11,22,33));
+		})
+		.TakeSnapshot()
+		.ModifyWorld([&](UWorld* World)
+		{
+			Actor->RemoveOwnedComponent(InstancedRemovedFromActor);
+			Actor->RemoveOwnedComponent(NativeRemovedFromActor);
+
+			InstancedAddedToActor = NewObject<UStaticMeshComponent>(Actor, UStaticMeshComponent::StaticClass(), "InstancedAddedToActor");
+			Actor->AddInstanceComponent(InstancedAddedToActor);
+			NativeAddedToActor = NewObject<UStaticMeshComponent>(Actor, UStaticMeshComponent::StaticClass(), "NativeAddedToActor");
+
+			InstancedStaysOnActor->SetRelativeLocation(FVector(100,200,300));
+			NativeStaysOnActor->SetRelativeLocation(FVector(110,220,330));
+		})
+	
+		.FilterProperties(Actor, [&](const FPropertySelectionMap& PropertySelectionMap)
+		{
+			const FRestorableObjectSelection ObjectSelection = PropertySelectionMap.GetObjectSelection(Actor);
+			
+			const FAddedAndRemovedComponentInfo* Selection = ObjectSelection.GetComponentSelection();
+			TestTrue(TEXT("Has Component Diff"), Selection != nullptr);
+			if (!Selection)
+			{
+				return;
+			}
+
+			TestTrue(TEXT("Added instanced actor should be removed"), Selection->EditorWorldComponentsToRemove.Contains(InstancedAddedToActor));
+			TestTrue(TEXT("Added native actor should be removed"), Selection->EditorWorldComponentsToRemove.Contains(NativeAddedToActor));
+			TestEqual(TEXT("Num actors to remove"), Selection->EditorWorldComponentsToRemove.Num(), 2);
+			
+			TestEqual(TEXT("Num actors to add back"), Selection->SnapshotComponentsToAdd.Num(), 2);
+			TSet<FString> NamesToCheck { "InstancedRemovedFromActor", "NativeRemovedFromActor" };
+			for (auto AddBackIt = Selection->SnapshotComponentsToAdd.CreateConstIterator(); AddBackIt; ++AddBackIt)
+			{
+				TestTrue(TEXT("Component should be added back"), NamesToCheck.Contains(AddBackIt->Get()->GetName()));
+			}
+
+			TestTrue(TEXT("Remaining instanced component has property selection"), PropertySelectionMap.GetObjectSelection(InstancedStaysOnActor).GetPropertySelection() != nullptr);
+			TestTrue(TEXT("Remaining native component has property selection"), PropertySelectionMap.GetObjectSelection(NativeStaysOnActor).GetPropertySelection() != nullptr);
+		});
+
+	// HasOriginalChangedPropertiesSinceSnapshotWasTaken return true
+	AActor* ActorWithAddedComponent = nullptr;
+	UActorComponent* AddedComponent = nullptr;
+	AActor* ActorWithRemovedComponent = nullptr;
+	FSnapshotTestRunner()
+		.ModifyWorld([&](UWorld* World)
+		{
+			ActorWithAddedComponent = World->SpawnActor<AActor>();
+			ActorWithRemovedComponent = World->SpawnActor<AActor>();
+
+			AddedComponent = NewObject<UStaticMeshComponent>(ActorWithRemovedComponent, UStaticMeshComponent::StaticClass(), "AddedComponent");
+			ActorWithRemovedComponent->AddInstanceComponent(AddedComponent);
+		})
+		.TakeSnapshot()
+
+		.ModifyWorld([&](UWorld* World)
+		{
+			ActorWithRemovedComponent->RemoveOwnedComponent(AddedComponent);
+
+			AddedComponent = NewObject<UStaticMeshComponent>(ActorWithAddedComponent, UStaticMeshComponent::StaticClass(), "AddedComponent");
+			ActorWithAddedComponent->AddInstanceComponent(AddedComponent);
+		})
+		.AccessSnapshot([&](ULevelSnapshot* Snapshot)
+		{
+			TOptional<AActor*> SnapshotActorWithAddedComponent = Snapshot->GetDeserializedActor(ActorWithAddedComponent);
+			TOptional<AActor*> SnapshotActorWithRemovedComponent = Snapshot->GetDeserializedActor(ActorWithRemovedComponent);
+			if (!ensure(SnapshotActorWithAddedComponent.IsSet() && SnapshotActorWithRemovedComponent.IsSet()))
+			{
+				return;
+			}
+
+			TestTrue(TEXT("ActorWithAddedComponent: HasOriginalChangedPropertiesSinceSnapshotWasTaken"), Snapshot->HasOriginalChangedPropertiesSinceSnapshotWasTaken(*SnapshotActorWithAddedComponent , ActorWithAddedComponent));
+			TestTrue(TEXT("ActorWithRemovedComponent: HasOriginalChangedPropertiesSinceSnapshotWasTaken"), Snapshot->HasOriginalChangedPropertiesSinceSnapshotWasTaken(*SnapshotActorWithRemovedComponent , ActorWithRemovedComponent));
+		});
+	
+	return true;
+}
+
+/**
+ * Given the below class, this tests that after applying to world MeshComponent points to the component in the world.
+ * This used to be a bug: the reference used to point to the equivalent component in the snapshot world or a class default.
+ * 
+ * class AActor
+ * {
+ *		UPROPERTY()
+ *		UStaticMeshComponent* MeshComponent;
+ * }
+ */
+IMPLEMENT_SIMPLE_AUTOMATION_TEST(FAreComponentReferenceValidAfterRestore, "VirtualProduction.LevelSnapshots.Snapshot.Component.AreComponentReferenceValidAfterRestore", (EAutomationTestFlags::ApplicationContextMask | EAutomationTestFlags::EngineFilter));
+bool FAreComponentReferenceValidAfterRestore::RunTest(const FString& Parameters)
+{
+	ASnapshotTestActor* Actor = nullptr;
+	
+	UInstancedStaticMeshComponent* OriginalMeshComp = nullptr;
+	UPointLightComponent* OriginalLightComp = nullptr;
+	
+	FSnapshotTestRunner()
+		.ModifyWorld([&](UWorld* World)
+		{
+			Actor = World->SpawnActor<ASnapshotTestActor>();
+			
+			OriginalMeshComp = Actor->InstancedMeshComponent;
+			OriginalLightComp = Actor->PointLightComponent;
+		})
+		.TakeSnapshot()
+		.ModifyWorld([&](UWorld* World)
+		{
+			Actor->PointLightComponent = nullptr;
+		})
+		.ApplySnapshot()
+		.RunTest([&]()
+		{
+			TestEqual(TEXT("Component reference was not changed"), Actor->InstancedMeshComponent, OriginalMeshComp);
+			TestEqual(TEXT("Component reference is correctly restored"), Actor->PointLightComponent, OriginalLightComp);
+		});
+	
+	return true;
+}
+
+/**
+ * Checks that component references as the one below do not show up in the selection set.
+ * 
+ * class AActor
+ * {
+ *		UPROPERTY()
+ *		UStaticMeshComponent* MeshComponent;
+ * }
+ */
+IMPLEMENT_SIMPLE_AUTOMATION_TEST(FEquivalentObjectReferencesDoNotHaveDiff, "VirtualProduction.LevelSnapshots.Snapshot.Component.EquivalentObjectReferencesDoNotHaveDiff", (EAutomationTestFlags::ApplicationContextMask | EAutomationTestFlags::EngineFilter));
+bool FEquivalentObjectReferencesDoNotHaveDiff::RunTest(const FString& Parameters)
+{
+	ASnapshotTestActor* Actor = nullptr;
+	
+	FSnapshotTestRunner()
+		.ModifyWorld([&](UWorld* World)
+		{
+			Actor = World->SpawnActor<ASnapshotTestActor>();
+		})
+		.TakeSnapshot()
+		.ModifyWorld([&](UWorld* World)
+		{
+			// Just here so we definitely detect changes on the actor
+			Actor->IntProperty = 12345;
+		})
+		.FilterProperties(Actor, [&](const FPropertySelectionMap& PropertySelectionMap)
+		{
+			const FPropertySelection* SelectedSubobjectProperties = PropertySelectionMap.GetObjectSelection(Actor).GetPropertySelection();
+
+			const FProperty* PointLightComponent = ASnapshotTestActor::StaticClass()->FindPropertyByName(GET_MEMBER_NAME_CHECKED(ASnapshotTestActor, PointLightComponent));
+			const FProperty* CubeMesh = ASnapshotTestActor::StaticClass()->FindPropertyByName(GET_MEMBER_NAME_CHECKED(ASnapshotTestActor, CubeMesh));
+			
+			TestTrue(TEXT("Component does not have diff"), !SelectedSubobjectProperties->IsPropertySelected(nullptr, PointLightComponent));
+			TestTrue(TEXT("Asset does not have diff"), !SelectedSubobjectProperties->IsPropertySelected(nullptr, CubeMesh));
+		});
+
+	return true;
+}
+
+/**
+ * Tests that an instanced component is:
+ * - recreated if it was removed
+ * - serialized into if it still exists
+ */
+IMPLEMENT_SIMPLE_AUTOMATION_TEST(FRestoreInstancedComponent, "VirtualProduction.LevelSnapshots.Snapshot.Component.RestoreInstancedComponent", (EAutomationTestFlags::ApplicationContextMask | EAutomationTestFlags::EngineFilter));
+bool FRestoreInstancedComponent::RunTest(const FString& Parameters)
+{
+	ASnapshotTestActor* Actor = nullptr;
+	UStaticMeshComponent* InstancedStaysOnActor = nullptr;
+	UStaticMeshComponent* InstancedRemovedFromActor = nullptr; 
+	UStaticMeshComponent* InstancedAddedToActor = nullptr; 
+	
+	FSnapshotTestRunner()
+		.ModifyWorld([&](UWorld* World)
+		{
+			Actor = World->SpawnActor<ASnapshotTestActor>();
+
+			// This component will stay on the actor
+			InstancedStaysOnActor = NewObject<UStaticMeshComponent>(Actor, UStaticMeshComponent::StaticClass(), "InstancedStaysOnActor");
+			Actor->AddInstanceComponent(InstancedStaysOnActor);
+			
+			InstancedStaysOnActor->SetStaticMesh(Actor->CubeMesh);
+			InstancedStaysOnActor->SetRelativeLocation(FVector(100, 200, 300));
+			InstancedStaysOnActor->AttachToComponent(Actor->GetMesh(), FAttachmentTransformRules::KeepRelativeTransform);
+
+
+			// This component will be recreated
+			InstancedRemovedFromActor = NewObject<UStaticMeshComponent>(Actor, UStaticMeshComponent::StaticClass(), "InstancedRemovedFromActor");
+			Actor->AddInstanceComponent(InstancedRemovedFromActor);
+
+			InstancedRemovedFromActor->SetStaticMesh(Actor->CylinderMesh);
+			InstancedRemovedFromActor->SetRelativeLocation(FVector(1, 2, 3));
+			InstancedRemovedFromActor->AttachToComponent(Actor->GetMesh(), FAttachmentTransformRules::KeepRelativeTransform);
+		})
+		.TakeSnapshot()
+		.ModifyWorld([&](UWorld* World)
+		{
+			InstancedStaysOnActor->SetStaticMesh(nullptr);
+			InstancedStaysOnActor->AttachToComponent(Actor->GetRootComponent(), FAttachmentTransformRules::KeepRelativeTransform);
+			InstancedStaysOnActor->SetRelativeLocation(FVector(1000, 2000, 3000));
+			
+			InstancedRemovedFromActor->AttachToComponent(Actor->GetRootComponent(), FAttachmentTransformRules::KeepRelativeTransform);
+			InstancedRemovedFromActor->SetRelativeLocation(FVector(10, 20, 30));
+			InstancedRemovedFromActor->DestroyComponent();
+
+			InstancedAddedToActor = NewObject<UStaticMeshComponent>(Actor, UStaticMeshComponent::StaticClass(), "InstancedAddedToActor");
+			Actor->AddInstanceComponent(InstancedAddedToActor);
+		})
+
+		// Exists on snapshot?
+		.AccessSnapshot([&](ULevelSnapshot* Snapshot)
+		{
+			TOptional<AActor*> CounterpartActor = Snapshot->GetDeserializedActor(Actor);
+			if (!ensureAlways(CounterpartActor))
+			{
+				return;
+			}
+
+			UActorComponent* const* CounterpartComponent = CounterpartActor.GetValue()->GetInstanceComponents().FindByPredicate([](UActorComponent* Component)
+			{
+				return Component->GetFName() == "InstancedStaysOnActor";
+			});
+			TestTrue(TEXT("Instanced component in InstancedComponents array on snapshot counterpart"), CounterpartComponent != nullptr);
+			if (!CounterpartComponent)
+			{
+				return;
+			}
+
+			USceneComponent* AsSceneComponent = Cast<USceneComponent>(*CounterpartComponent);
+			ASnapshotTestActor* CastedCounterpartActor = Cast<ASnapshotTestActor>(*CounterpartActor);
+			check(AsSceneComponent);
+			TestTrue(TEXT("Snapshot attach parent correct"), AsSceneComponent->GetAttachParent() == CastedCounterpartActor->GetMesh());
+		})
+
+		// Restores correctly?
+		.ApplySnapshot()
+		.RunTest([&]()
+		{
+			// Instanced component that stayed
+			TestEqual(TEXT("Instanced component static mesh was restored"), ToRawPtr(InstancedStaysOnActor->GetStaticMesh()), Actor->CubeMesh);
+			TestEqual(TEXT("Instanced component relative location was restored"), InstancedStaysOnActor->GetRelativeLocation(), FVector(100, 200, 300));
+			TestTrue(TEXT("Instanced component attach parent was restored"), InstancedStaysOnActor->GetAttachParent() == Actor->GetMesh());
+
+
+			// Component was removed
+			UActorComponent* const* RemovedComponent = Actor->GetInstanceComponents().FindByPredicate([&](UActorComponent* Component)
+			{
+				return Component->GetFName() == "InstancedAddedToActor";
+			});
+			TestTrue(TEXT("Instanced component was removed from actor"), RemovedComponent == nullptr);
+			TestTrue(TEXT("Removed component is marked from destroy"), InstancedAddedToActor->IsPendingKill());
+
+			// Component was added back
+			UActorComponent* const* PossibleRecreatedComponent = Actor->GetInstanceComponents().FindByPredicate([&](UActorComponent* Component)
+			{
+				return Component->GetFName() == "InstancedRemovedFromActor";
+			});
+			if (!PossibleRecreatedComponent)
+			{
+				AddError(TEXT("Removed instanced component was not added back."));
+				return;
+			}
+			UStaticMeshComponent* RecreatedComponent = Cast<UStaticMeshComponent>(*PossibleRecreatedComponent);
+			TestEqual(TEXT("Recreated component has correct mesh"), ToRawPtr(RecreatedComponent->GetStaticMesh()), Actor->CylinderMesh);
+			TestEqual(TEXT("Recreated component has correct relative location"), RecreatedComponent->GetRelativeLocation(), FVector(1, 2, 3));
+			TestTrue(TEXT("Recreated component attach parent was restored"), RecreatedComponent->GetAttachParent() == Actor->GetMesh());
+		});
+
+	return true;
+}
+
+IMPLEMENT_SIMPLE_AUTOMATION_TEST(FRecreatedActorHasAllComponents, "VirtualProduction.LevelSnapshots.Snapshot.Component.RecreatedActorHasAllComponents", (EAutomationTestFlags::ApplicationContextMask | EAutomationTestFlags::EngineFilter));
+bool FRecreatedActorHasAllComponents::RunTest(const FString& Parameters)
+{
+	AActor* Actor = nullptr;
+	UStaticMeshComponent* InstancedRemovedFromActor = nullptr;
+	UStaticMeshComponent* NativeRemovedFromActor = nullptr;
+	
+	FSnapshotTestRunner()
+		.ModifyWorld([&](UWorld* World)
+		{
+			FActorSpawnParameters SpawnParams;
+			SpawnParams.Name = "TestActor";
+			Actor = World->SpawnActor<AActor>(SpawnParams);
+
+			InstancedRemovedFromActor = NewObject<UStaticMeshComponent>(Actor, UStaticMeshComponent::StaticClass(), "InstancedRemovedFromActor");
+			Actor->AddInstanceComponent(InstancedRemovedFromActor);
+			NativeRemovedFromActor = NewObject<UStaticMeshComponent>(Actor, UStaticMeshComponent::StaticClass(), "NativeRemovedFromActor");
+
+			InstancedRemovedFromActor->SetRelativeLocation(FVector(10,20,30));
+			NativeRemovedFromActor->SetRelativeLocation(FVector(11,22,33));
+		})
+		.TakeSnapshot()
+		.ModifyWorld([&](UWorld* World)
+		{
+			World->DestroyActor(Actor);
+			Actor = nullptr;
+		})
+		.ApplySnapshot()
+		.ModifyWorld([&](UWorld* World)
+		{
+			for (TActorIterator<AActor> It(World); It; ++It)
+			{
+				if (It->GetName().Equals("TestActor"))
+				{
+					Actor = *It;
+				}
+			}
+
+			if (!ensure(Actor))
+			{
+				return;
+			}
+
+			TInlineComponentArray<UActorComponent*> Components(Actor);
+			UActorComponent** InstancedComp = Components.FindByPredicate([](UActorComponent* Component) { return Component->GetName().Equals("InstancedRemovedFromActor"); });
+			UActorComponent** NativeComp = Components.FindByPredicate([](UActorComponent* Component) { return Component->GetName().Equals("NativeRemovedFromActor"); });
+
+			TestTrue(TEXT("Instance Component was recreated"), InstancedComp != nullptr);
+			TestTrue(TEXT("Native Component was recreated"), NativeComp != nullptr);
+			if (!InstancedComp || !NativeComp)
+			{
+				return;
+			}
+			
+			TestEqual(TEXT("Instance Component had properties restored"), Cast<UStaticMeshComponent>(*InstancedComp)->GetRelativeLocation(), FVector(10,20,30));
+			TestEqual(TEXT("Native Component had properties restored"), Cast<UStaticMeshComponent>(*NativeComp)->GetRelativeLocation(), FVector(11,22,33));
+		});
+	
+	return true;
+}
diff --git a/Engine/Plugins/VirtualProduction/LevelSnapshots/Source/LevelSnapshots/Private/Tests/CustomObjectSerializationTests.cpp b/Engine/Plugins/VirtualProduction/LevelSnapshots/Source/LevelSnapshots/Private/Tests/CustomObjectSerializationTests.cpp
--- a/Engine/Plugins/VirtualProduction/LevelSnapshots/Source/LevelSnapshots/Private/Tests/CustomObjectSerializationTests.cpp
+++ b/Engine/Plugins/VirtualProduction/LevelSnapshots/Source/LevelSnapshots/Private/Tests/CustomObjectSerializationTests.cpp
@@ -1,15 +1,46 @@
 ﻿// Copyright Epic Games, Inc. All Rights Reserved.
 
+#include "CustomSubobjectRestorationInfo.h"
 #include "LevelSnapshotsModule.h"
 #include "PropertySelectionMap.h"
 #include "Serialization/ICustomObjectSnapshotSerializer.h"
 #include "Serialization/ObjectSnapshotSerializationData.h"
 #include "SnapshotTestActor.h"
 #include "SnapshotTestRunner.h"
+#include "Engine/StaticMeshActor.h"
 
 #include "Misc/AutomationTest.h"
 #include "Misc/ScopeExit.h"
 
+namespace
+{
+	void BlacklistSubobjectProperties()
+	{
+		FLevelSnapshotsModule& Module = FModuleManager::Get().GetModuleChecked<FLevelSnapshotsModule>("LevelSnapshots");
+		TSet<const FProperty*> BlacklistedProperties;
+		
+		BlacklistedProperties.Add(ASnapshotTestActor::StaticClass()->FindPropertyByName(GET_MEMBER_NAME_CHECKED(ASnapshotTestActor, EditableInstancedSubobject_DefaultSubobject)));
+		BlacklistedProperties.Add(ASnapshotTestActor::StaticClass()->FindPropertyByName(GET_MEMBER_NAME_CHECKED(ASnapshotTestActor, InstancedOnlySubobject_DefaultSubobject)));
+		BlacklistedProperties.Add(ASnapshotTestActor::StaticClass()->FindPropertyByName(GET_MEMBER_NAME_CHECKED(ASnapshotTestActor, NakedSubobject_DefaultSubobject)));
+		BlacklistedProperties.Add(USubobject::StaticClass()->FindPropertyByName(GET_MEMBER_NAME_CHECKED(USubobject, NestedChild)));
+
+		Module.AddBlacklistedProperties(BlacklistedProperties);
+	}
+
+	void UnblacklistSubobjectProperties()
+	{
+		FLevelSnapshotsModule& Module = FModuleManager::Get().GetModuleChecked<FLevelSnapshotsModule>("LevelSnapshots");
+		TSet<const FProperty*> BlacklistedProperties;
+		
+		BlacklistedProperties.Add(ASnapshotTestActor::StaticClass()->FindPropertyByName(GET_MEMBER_NAME_CHECKED(ASnapshotTestActor, EditableInstancedSubobject_DefaultSubobject)));
+		BlacklistedProperties.Add(ASnapshotTestActor::StaticClass()->FindPropertyByName(GET_MEMBER_NAME_CHECKED(ASnapshotTestActor, InstancedOnlySubobject_DefaultSubobject)));
+		BlacklistedProperties.Add(ASnapshotTestActor::StaticClass()->FindPropertyByName(GET_MEMBER_NAME_CHECKED(ASnapshotTestActor, NakedSubobject_DefaultSubobject)));
+		BlacklistedProperties.Add(USubobject::StaticClass()->FindPropertyByName(GET_MEMBER_NAME_CHECKED(USubobject, NestedChild)));
+
+		Module.RemoveBlacklistedProperties(BlacklistedProperties);
+	}
+}
+
 /**
  * Tests all interface functions are called at the correct time.
  */
@@ -42,7 +73,7 @@ bool FRestoreActorCustomSubobject::RunTest(const FString& Parameters)
 			CallOrder.Add(EFunctionCall::OnTakeSnapshot);
 			
 			Test.TestTrue(TEXT("Correct editor object passed in"), EditorObject == TestActor);
-			const int32 Index = DataStorage.AddSubobjectSnapshot(TestActor->InstancedSubobject);
+			const int32 Index = DataStorage.AddSubobjectSnapshot(TestActor->InstancedOnlySubobject_DefaultSubobject);
 			DataStorage.GetSubobjectMetaData(Index)->WriteObjectAnnotation(FObjectAnnotator::CreateLambda([](FArchive& Archive)
 			{
 				int32 TestSubobjectInfo = 42;
@@ -62,7 +93,7 @@ bool FRestoreActorCustomSubobject::RunTest(const FString& Parameters)
 
 			if (ASnapshotTestActor* SnapshotActor = Cast<ASnapshotTestActor>(SnapshotObject))
 			{
-				return SnapshotActor->InstancedSubobject;
+				return SnapshotActor->InstancedOnlySubobject_DefaultSubobject;
 			}
 			
 			Test.AddError(TEXT("Expected SnapshotObject to be an instance of ASnapshotTestActor"));
@@ -74,7 +105,7 @@ bool FRestoreActorCustomSubobject::RunTest(const FString& Parameters)
 			CallOrder.Add(EFunctionCall::FindOrRecreateSubobjectInEditorWorld);
 			
 			Test.TestTrue(TEXT("Correct editor object passed in"), EditorObject == TestActor);
-			return TestActor->InstancedSubobject;	
+			return TestActor->InstancedOnlySubobject_DefaultSubobject;	
 		}
 		
 		virtual UObject* FindSubobjectInEditorWorld(UObject* EditorObject, const ISnapshotSubobjectMetaData& ObjectData, const ICustomSnapshotSerializationData& DataStorage) override
@@ -97,7 +128,7 @@ bool FRestoreActorCustomSubobject::RunTest(const FString& Parameters)
 			}));
 			Test.TestEqual(TEXT("Saved custom actor data is correct"), TestActorInfo, 21);
 			
-			return TestActor->InstancedSubobject;	
+			return TestActor->InstancedOnlySubobject_DefaultSubobject;	
 		}
 
 		
@@ -128,9 +159,11 @@ bool FRestoreActorCustomSubobject::RunTest(const FString& Parameters)
 	TSharedRef<FStub> Stub = MakeShared<FStub>(*this);
 	FLevelSnapshotsModule& Module = FModuleManager::Get().GetModuleChecked<FLevelSnapshotsModule>("LevelSnapshots");
 	Module.RegisterCustomObjectSerializer(ASnapshotTestActor::StaticClass(), Stub);
+	BlacklistSubobjectProperties();
 	ON_SCOPE_EXIT
 	{
 		Module.UnregisterCustomObjectSerializer(ASnapshotTestActor::StaticClass());
+		UnblacklistSubobjectProperties();
 	};
 	
 	FSnapshotTestRunner()
@@ -138,8 +171,8 @@ bool FRestoreActorCustomSubobject::RunTest(const FString& Parameters)
 		{
 			Stub->TestActor = World->SpawnActor<ASnapshotTestActor>();
 
-			Stub->TestActor->InstancedSubobject->FloatProperty = 42.f;
-			Stub->TestActor->InstancedSubobject->IntProperty = 21;
+			Stub->TestActor->InstancedOnlySubobject_DefaultSubobject->FloatProperty = 42.f;
+			Stub->TestActor->InstancedOnlySubobject_DefaultSubobject->IntProperty = 21;
 			Stub->TestActor->IntProperty = 42;
 			Stub->TestActor->TestComponent->IntProperty = 42;
 		})
@@ -147,8 +180,8 @@ bool FRestoreActorCustomSubobject::RunTest(const FString& Parameters)
 	
 		.ModifyWorld([&](UWorld* World)
 		{
-			Stub->TestActor->InstancedSubobject->FloatProperty = 420.f;
-			Stub->TestActor->InstancedSubobject->IntProperty = 210;
+			Stub->TestActor->InstancedOnlySubobject_DefaultSubobject->FloatProperty = 420.f;
+			Stub->TestActor->InstancedOnlySubobject_DefaultSubobject->IntProperty = 210;
 			Stub->TestActor->IntProperty = 420;
 			Stub->TestActor->TestComponent->IntProperty = 420;
 		})
@@ -172,8 +205,8 @@ bool FRestoreActorCustomSubobject::RunTest(const FString& Parameters)
 			TestTrue(TEXT("PreApplySnapshotProperties called before PostApplySnapshotPropertiesIndex"), PreApplySnapshotPropertiesIndex < OnPostSerializeEditorSubobjectIndex);
 
 			// Custom subobject restored
-			TestEqual(TEXT("FloatProperty restored"), Stub->TestActor->InstancedSubobject->FloatProperty, 42.f);
-			TestEqual(TEXT("IntProperty restored"), Stub->TestActor->InstancedSubobject->IntProperty, 21);
+			TestEqual(TEXT("FloatProperty restored"), Stub->TestActor->InstancedOnlySubobject_DefaultSubobject->FloatProperty, 42.f);
+			TestEqual(TEXT("IntProperty restored"), Stub->TestActor->InstancedOnlySubobject_DefaultSubobject->IntProperty, 21);
 			// Normal values are still restored
 			TestEqual(TEXT("Actor property still restored"), Stub->TestActor->IntProperty, 42);
 			TestEqual(TEXT("Component property still restored"), Stub->TestActor->TestComponent->IntProperty, 42);
@@ -182,6 +215,145 @@ bool FRestoreActorCustomSubobject::RunTest(const FString& Parameters)
 	return true;
 }
 
+/**
+ * Makes sure that we can write
+ */
+IMPLEMENT_SIMPLE_AUTOMATION_TEST(FSaveAndLoadObjectAnnotation, "VirtualProduction.LevelSnapshots.Snapshot.CustomObjectSerialization.SaveAndLoadObjectAnnotation", (EAutomationTestFlags::ApplicationContextMask | EAutomationTestFlags::EngineFilter));
+bool FSaveAndLoadObjectAnnotation::RunTest(const FString& Parameters)
+{
+	class FStub : public ICustomObjectSnapshotSerializer
+	{
+		FSaveAndLoadObjectAnnotation& Test;
+	public:
+
+		FStub(FSaveAndLoadObjectAnnotation& Test)
+			: Test(Test)
+		{}
+		
+		virtual void OnTakeSnapshot(UObject* EditorObject, ICustomSnapshotSerializationData& DataStorage) override
+		{
+			ASnapshotTestActor* Actor = Cast<ASnapshotTestActor>(EditorObject);
+			
+			const int32 Index = DataStorage.AddSubobjectSnapshot(Actor->NonReflectedSubobject);
+			DataStorage.WriteObjectAnnotation(FObjectAnnotator::CreateLambda([Actor](FArchive& Archive)
+			{
+				Archive << Actor->NonReflectedName;
+				Archive << Actor->NonReflectedObjectProperty;
+				Archive << Actor->NonReflectedSoftPtr;
+			}));
+			DataStorage.GetSubobjectMetaData(Index)->WriteObjectAnnotation(FObjectAnnotator::CreateLambda([Actor](FArchive& Archive)
+			{
+				Archive << Actor->NonReflectedSubobject->NonReflectedName;
+				Archive << Actor->NonReflectedSubobject->NonReflectedObjectProperty;
+				Archive << Actor->NonReflectedSubobject->NonReflectedSoftPtr;
+			}));
+		}
+
+		virtual UObject* FindOrRecreateSubobjectInSnapshotWorld(UObject* SnapshotObject, const ISnapshotSubobjectMetaData& ObjectData, const ICustomSnapshotSerializationData& DataStorage) override
+		{
+			ASnapshotTestActor* Actor = Cast<ASnapshotTestActor>(SnapshotObject);
+			Actor->AllocateNonReflectedSubobject();
+			return Actor->NonReflectedSubobject;
+		}
+		
+		virtual UObject* FindOrRecreateSubobjectInEditorWorld(UObject* EditorObject, const ISnapshotSubobjectMetaData& ObjectData, const ICustomSnapshotSerializationData& DataStorage) override
+		{
+			ASnapshotTestActor* Actor = Cast<ASnapshotTestActor>(EditorObject);
+			Actor->AllocateNonReflectedSubobject();
+			return Actor->NonReflectedSubobject;
+		}
+		
+		virtual UObject* FindSubobjectInEditorWorld(UObject* EditorObject, const ISnapshotSubobjectMetaData& ObjectData, const ICustomSnapshotSerializationData& DataStorage) override
+		{
+			
+			ASnapshotTestActor* Actor = Cast<ASnapshotTestActor>(EditorObject);
+			return Actor->NonReflectedSubobject;
+		}
+
+		virtual void PostApplySnapshotProperties(UObject* OriginalObject, const ICustomSnapshotSerializationData& DataStorage) override
+		{
+			ASnapshotTestActor* Actor = Cast<ASnapshotTestActor>(OriginalObject);
+			
+			DataStorage.ReadObjectAnnotation(FObjectAnnotator::CreateLambda([Actor](FArchive& Archive)
+			{
+				Archive << Actor->NonReflectedName;
+				Archive << Actor->NonReflectedObjectProperty;
+				Archive << Actor->NonReflectedSoftPtr;
+			}));
+
+			DataStorage.GetSubobjectMetaData(0)->ReadObjectAnnotation(FObjectAnnotator::CreateLambda([Actor](FArchive& Archive)
+			{
+				Archive << Actor->NonReflectedSubobject->NonReflectedName;
+				Archive << Actor->NonReflectedSubobject->NonReflectedObjectProperty;
+				Archive << Actor->NonReflectedSubobject->NonReflectedSoftPtr;
+			}));
+		}
+	};
+
+	// Handle registering and unregistering of custom serializer
+	TSharedRef<FStub> Stub = MakeShared<FStub>(*this);
+	FLevelSnapshotsModule& Module = FModuleManager::Get().GetModuleChecked<FLevelSnapshotsModule>("LevelSnapshots");
+	Module.RegisterCustomObjectSerializer(ASnapshotTestActor::StaticClass(), Stub);
+	BlacklistSubobjectProperties();
+	ON_SCOPE_EXIT
+	{
+		Module.UnregisterCustomObjectSerializer(ASnapshotTestActor::StaticClass());
+		UnblacklistSubobjectProperties();
+	};
+
+	ASnapshotTestActor* TestActor = nullptr;
+	AStaticMeshActor* FirstReferencedActor = nullptr;
+	AStaticMeshActor* SecondReferencedActor = nullptr;
+	
+	FSnapshotTestRunner()
+		.ModifyWorld([&](UWorld* World)
+		{
+			TestActor = World->SpawnActor<ASnapshotTestActor>();
+			TestActor->AllocateNonReflectedSubobject();
+			FirstReferencedActor = World->SpawnActor<AStaticMeshActor>();
+			SecondReferencedActor = World->SpawnActor<AStaticMeshActor>();
+			
+			TestActor->IntProperty = 1;
+			TestActor->NonReflectedName = FName("TestNonReflectedName_OnActor");
+			TestActor->NonReflectedObjectProperty = FirstReferencedActor;
+			TestActor->NonReflectedSoftPtr = FirstReferencedActor;
+			
+			TestActor->NonReflectedSubobject->NonReflectedName = FName("TestNonReflectedName_OnSubobject");
+			TestActor->NonReflectedSubobject->NonReflectedObjectProperty = SecondReferencedActor;
+			TestActor->NonReflectedSubobject->NonReflectedSoftPtr = SecondReferencedActor;
+		})
+		.TakeSnapshot()
+	
+		.ModifyWorld([&](UWorld* World)
+		{
+			// This makes the actor show up as changed - otherwise nothing will get serialized because all other properties are not reflected
+			TestActor->IntProperty = 2;
+
+			TestActor->NonReflectedName = FName(NAME_None);
+			TestActor->NonReflectedObjectProperty = nullptr;
+			TestActor->NonReflectedSoftPtr = nullptr;
+			
+			TestActor->NonReflectedSubobject->NonReflectedName = FName(NAME_None);
+			TestActor->NonReflectedSubobject->NonReflectedObjectProperty = nullptr;
+			TestActor->NonReflectedSubobject->NonReflectedSoftPtr = nullptr;
+		})
+		.ApplySnapshot()
+
+		.RunTest([&]()
+		{
+			TestEqual(TEXT("TestActor->NonReflectedName"), TestActor->NonReflectedName, FName("TestNonReflectedName_OnActor"));
+			TestTrue(TEXT("TestActor->NonReflectedObjectProperty"), TestActor->NonReflectedObjectProperty == FirstReferencedActor);
+			TestTrue(TEXT("TestActor->NonReflectedSoftPtr"), TestActor->NonReflectedSoftPtr == FirstReferencedActor);
+			
+			TestEqual(TEXT("TestActor->NonReflectedSubobject->NonReflectedName"), TestActor->NonReflectedSubobject->NonReflectedName, FName("TestNonReflectedName_OnSubobject"));
+			TestTrue(TEXT("TestActor->NonReflectedSubobject->NonReflectedObjectProperty"), TestActor->NonReflectedSubobject->NonReflectedObjectProperty == SecondReferencedActor);
+			TestTrue(TEXT("TestActor->NonReflectedSubobject->NonReflectedSoftPtr"), TestActor->NonReflectedSubobject->NonReflectedSoftPtr == SecondReferencedActor);
+			
+		});
+
+	return true;
+}
+
 /**
  * Tests that custom serialization works when an actor adds a subobject dependency to a subobject that also has a custom serializer.
  */
@@ -194,14 +366,14 @@ bool FRestoreNestedCustomSubobject::RunTest(const FString& Parameters)
 		
 		virtual void OnTakeSnapshot(UObject* EditorObject, ICustomSnapshotSerializationData& DataStorage) override
 		{
-			DataStorage.AddSubobjectSnapshot(TestActor->InstancedSubobject);
+			DataStorage.AddSubobjectSnapshot(TestActor->InstancedOnlySubobject_DefaultSubobject);
 		}
 
 		virtual UObject* FindOrRecreateSubobjectInSnapshotWorld(UObject* SnapshotObject, const ISnapshotSubobjectMetaData& ObjectData, const ICustomSnapshotSerializationData& DataStorage) override
 		{
 			if (ASnapshotTestActor* SnapshotActor = Cast<ASnapshotTestActor>(SnapshotObject))
 			{
-				return SnapshotActor->InstancedSubobject;
+				return SnapshotActor->InstancedOnlySubobject_DefaultSubobject;
 			}
 
 			checkNoEntry();
@@ -210,12 +382,12 @@ bool FRestoreNestedCustomSubobject::RunTest(const FString& Parameters)
 		
 		virtual UObject* FindOrRecreateSubobjectInEditorWorld(UObject* EditorObject, const ISnapshotSubobjectMetaData& ObjectData, const ICustomSnapshotSerializationData& DataStorage) override
 		{
-			return TestActor->InstancedSubobject;	
+			return TestActor->InstancedOnlySubobject_DefaultSubobject;	
 		}
 		
 		virtual UObject* FindSubobjectInEditorWorld(UObject* EditorObject, const ISnapshotSubobjectMetaData& ObjectData, const ICustomSnapshotSerializationData& DataStorage) override
 		{
-			return TestActor->InstancedSubobject;	
+			return TestActor->InstancedOnlySubobject_DefaultSubobject;	
 		}
 
 		ASnapshotTestActor* TestActor;
@@ -227,7 +399,7 @@ bool FRestoreNestedCustomSubobject::RunTest(const FString& Parameters)
 		
 		virtual void OnTakeSnapshot(UObject* EditorObject, ICustomSnapshotSerializationData& DataStorage) override
 		{
-			DataStorage.AddSubobjectSnapshot(TestActor->InstancedSubobject->NestedChild);
+			DataStorage.AddSubobjectSnapshot(TestActor->InstancedOnlySubobject_DefaultSubobject->NestedChild);
 		}
 
 		virtual UObject* FindOrRecreateSubobjectInSnapshotWorld(UObject* SnapshotObject, const ISnapshotSubobjectMetaData& ObjectData, const ICustomSnapshotSerializationData& DataStorage) override
@@ -243,12 +415,12 @@ bool FRestoreNestedCustomSubobject::RunTest(const FString& Parameters)
 		
 		virtual UObject* FindOrRecreateSubobjectInEditorWorld(UObject* EditorObject, const ISnapshotSubobjectMetaData& ObjectData, const ICustomSnapshotSerializationData& DataStorage) override
 		{
-			return TestActor->InstancedSubobject->NestedChild;	
+			return TestActor->InstancedOnlySubobject_DefaultSubobject->NestedChild;	
 		}
 		
 		virtual UObject* FindSubobjectInEditorWorld(UObject* EditorObject, const ISnapshotSubobjectMetaData& ObjectData, const ICustomSnapshotSerializationData& DataStorage) override
 		{
-			return TestActor->InstancedSubobject->NestedChild;	
+			return TestActor->InstancedOnlySubobject_DefaultSubobject->NestedChild;	
 		}
 
 		ASnapshotTestActor* TestActor;
@@ -260,10 +432,12 @@ bool FRestoreNestedCustomSubobject::RunTest(const FString& Parameters)
 	FLevelSnapshotsModule& Module = FModuleManager::Get().GetModuleChecked<FLevelSnapshotsModule>("LevelSnapshots");
 	Module.RegisterCustomObjectSerializer(ASnapshotTestActor::StaticClass(), ActorSerializer);
 	Module.RegisterCustomObjectSerializer(USubobject::StaticClass(), SubobjectSerializer);
+	BlacklistSubobjectProperties();
 	ON_SCOPE_EXIT
 	{
 		Module.UnregisterCustomObjectSerializer(ASnapshotTestActor::StaticClass());
 		Module.UnregisterCustomObjectSerializer(USubobject::StaticClass());
+		UnblacklistSubobjectProperties();
 	};
 
 	ASnapshotTestActor* TestActor = nullptr;
@@ -276,22 +450,22 @@ bool FRestoreNestedCustomSubobject::RunTest(const FString& Parameters)
 			ActorSerializer->TestActor = TestActor;
 			SubobjectSerializer->TestActor = TestActor;
 			
-			TestActor->InstancedSubobject->FloatProperty = 42.f;
-			TestActor->InstancedSubobject->NestedChild->FloatProperty = 21.f;
+			TestActor->InstancedOnlySubobject_DefaultSubobject->FloatProperty = 42.f;
+			TestActor->InstancedOnlySubobject_DefaultSubobject->NestedChild->FloatProperty = 21.f;
 		})
 		.TakeSnapshot()
 	
 		.ModifyWorld([&](UWorld* World)
 		{
-			TestActor->InstancedSubobject->FloatProperty = 420.f;
-			TestActor->InstancedSubobject->NestedChild->FloatProperty = 210.f;
+			TestActor->InstancedOnlySubobject_DefaultSubobject->FloatProperty = 420.f;
+			TestActor->InstancedOnlySubobject_DefaultSubobject->NestedChild->FloatProperty = 210.f;
 		})
 		.ApplySnapshot()
 
 		.RunTest([&]()
 		{
-			TestEqual(TEXT("Subobject property restored"), TestActor->InstancedSubobject->FloatProperty, 42.f);
-			TestEqual(TEXT("Nested subobject property restored"), TestActor->InstancedSubobject->NestedChild->FloatProperty, 21.f);
+			TestEqual(TEXT("Subobject property restored"), TestActor->InstancedOnlySubobject_DefaultSubobject->FloatProperty, 42.f);
+			TestEqual(TEXT("Nested subobject property restored"), TestActor->InstancedOnlySubobject_DefaultSubobject->NestedChild->FloatProperty, 21.f);
 		});
 
 
@@ -303,19 +477,19 @@ bool FRestoreNestedCustomSubobject::RunTest(const FString& Parameters)
 			ActorSerializer->TestActor = TestActor;
 			SubobjectSerializer->TestActor = TestActor;
 			
-			TestActor->InstancedSubobject->NestedChild->FloatProperty = 21.f;
+			TestActor->InstancedOnlySubobject_DefaultSubobject->NestedChild->FloatProperty = 21.f;
 		})
 		.TakeSnapshot()
 	
 		.ModifyWorld([&](UWorld* World)
 		{
-			TestActor->InstancedSubobject->NestedChild->FloatProperty = 210.f;
+			TestActor->InstancedOnlySubobject_DefaultSubobject->NestedChild->FloatProperty = 210.f;
 		})
 		.ApplySnapshot()
 
 		.RunTest([&]()
 		{
-			TestEqual(TEXT("Nested subobject property restored when outer had no changed properties"), TestActor->InstancedSubobject->NestedChild->FloatProperty, 21.f);
+			TestEqual(TEXT("Nested subobject property restored when outer had no changed properties"), TestActor->InstancedOnlySubobject_DefaultSubobject->NestedChild->FloatProperty, 21.f);
 		});
 	
 	return true;
@@ -345,28 +519,28 @@ bool FNonEditableObjectPropertyNotRestoredByDefault::RunTest(const FString& Para
 			ChangedPropertiesActor = World->SpawnActor<ASnapshotTestActor>();
 			NulledActor = World->SpawnActor<ASnapshotTestActor>();
 
-			ChangedPropertiesActor->InstancedSubobject->FloatProperty = 42.f;
-			ChangedPropertiesActor->NakedSubobject->FloatProperty = 21.f;
+			ChangedPropertiesActor->InstancedOnlySubobject_DefaultSubobject->FloatProperty = 42.f;
+			ChangedPropertiesActor->NakedSubobject_DefaultSubobject->FloatProperty = 21.f;
 		})
 		.TakeSnapshot()
 	
 		.ModifyWorld([&](UWorld* World)
 		{
-			ChangedPropertiesActor->InstancedSubobject->FloatProperty = 420.f;
-			ChangedPropertiesActor->NakedSubobject->FloatProperty = 210.f;
+			ChangedPropertiesActor->InstancedOnlySubobject_DefaultSubobject->FloatProperty = 420.f;
+			ChangedPropertiesActor->NakedSubobject_DefaultSubobject->FloatProperty = 210.f;
 
-			NulledActor->InstancedSubobject = nullptr;
-			NulledActor->NakedSubobject = nullptr;
+			NulledActor->InstancedOnlySubobject_DefaultSubobject = nullptr;
+			NulledActor->NakedSubobject_DefaultSubobject = nullptr;
 		})
 		.ApplySnapshot()
 
 		.RunTest([&]()
 		{
-			TestEqual(TEXT("Unsupported instanced property not restored"), ChangedPropertiesActor->InstancedSubobject->FloatProperty, 420.f);
-			TestEqual(TEXT("Unsupported naked property not restored"), ChangedPropertiesActor->NakedSubobject->FloatProperty, 210.f);
+			TestEqual(TEXT("Unsupported instanced property not restored"), ChangedPropertiesActor->InstancedOnlySubobject_DefaultSubobject->FloatProperty, 420.f);
+			TestEqual(TEXT("Unsupported naked property not restored"), ChangedPropertiesActor->NakedSubobject_DefaultSubobject->FloatProperty, 210.f);
 			
-			TestTrue(TEXT("Unsupported instanced property stays null"), NulledActor->InstancedSubobject == nullptr);
-			TestTrue(TEXT("Unsupported naked property stays null"), NulledActor->NakedSubobject == nullptr);
+			TestTrue(TEXT("Unsupported instanced property stays null"), NulledActor->InstancedOnlySubobject_DefaultSubobject == nullptr);
+			TestTrue(TEXT("Unsupported naked property stays null"), NulledActor->NakedSubobject_DefaultSubobject == nullptr);
 		});
 
 	return true;
@@ -384,14 +558,14 @@ bool FFilterForPropertiesOnSubobjects::RunTest(const FString& Parameters)
 		
 		virtual void OnTakeSnapshot(UObject* EditorObject, ICustomSnapshotSerializationData& DataStorage) override
 		{
-			DataStorage.AddSubobjectSnapshot(TestActor->InstancedSubobject);
+			DataStorage.AddSubobjectSnapshot(TestActor->InstancedOnlySubobject_DefaultSubobject);
 		}
 
 		virtual UObject* FindOrRecreateSubobjectInSnapshotWorld(UObject* SnapshotObject, const ISnapshotSubobjectMetaData& ObjectData, const ICustomSnapshotSerializationData& DataStorage) override
 		{
 			if (ASnapshotTestActor* SnapshotActor = Cast<ASnapshotTestActor>(SnapshotObject))
 			{
-				return SnapshotActor->InstancedSubobject;
+				return SnapshotActor->InstancedOnlySubobject_DefaultSubobject;
 			}
 
 			checkNoEntry();
@@ -400,12 +574,12 @@ bool FFilterForPropertiesOnSubobjects::RunTest(const FString& Parameters)
 		
 		virtual UObject* FindOrRecreateSubobjectInEditorWorld(UObject* EditorObject, const ISnapshotSubobjectMetaData& ObjectData, const ICustomSnapshotSerializationData& DataStorage) override
 		{
-			return TestActor->InstancedSubobject;	
+			return TestActor->InstancedOnlySubobject_DefaultSubobject;	
 		}
 		
 		virtual UObject* FindSubobjectInEditorWorld(UObject* EditorObject, const ISnapshotSubobjectMetaData& ObjectData, const ICustomSnapshotSerializationData& DataStorage) override
 		{
-			return TestActor->InstancedSubobject;	
+			return TestActor->InstancedOnlySubobject_DefaultSubobject;	
 		}
 
 		ASnapshotTestActor* TestActor;
@@ -415,9 +589,11 @@ bool FFilterForPropertiesOnSubobjects::RunTest(const FString& Parameters)
 	TSharedRef<FStub> Stub = MakeShared<FStub>();
 	FLevelSnapshotsModule& Module = FModuleManager::Get().GetModuleChecked<FLevelSnapshotsModule>("LevelSnapshots");
 	Module.RegisterCustomObjectSerializer(ASnapshotTestActor::StaticClass(), Stub);
+	BlacklistSubobjectProperties();
 	ON_SCOPE_EXIT
 	{
 		Module.UnregisterCustomObjectSerializer(ASnapshotTestActor::StaticClass());
+		UnblacklistSubobjectProperties();
 	};
 
 
@@ -429,13 +605,13 @@ bool FFilterForPropertiesOnSubobjects::RunTest(const FString& Parameters)
 		{
 			Stub->TestActor = World->SpawnActor<ASnapshotTestActor>();
 			
-			Stub->TestActor->InstancedSubobject->FloatProperty = 42.f;
+			Stub->TestActor->InstancedOnlySubobject_DefaultSubobject->FloatProperty = 42.f;
 		})
 		.TakeSnapshot()
 	
 		.ModifyWorld([&](UWorld* World)
 		{
-			Stub->TestActor->InstancedSubobject->FloatProperty = 420.f;
+			Stub->TestActor->InstancedOnlySubobject_DefaultSubobject->FloatProperty = 420.f;
 		})
 
 		.FilterProperties(Stub->TestActor, [&](const FPropertySelectionMap& PropertySelectionMap)
@@ -444,14 +620,14 @@ bool FFilterForPropertiesOnSubobjects::RunTest(const FString& Parameters)
 			UClass* SubobjectClass = USubobject::StaticClass();
 			const FProperty* ChangedSubobjectProperty = SubobjectClass->FindPropertyByName(GET_MEMBER_NAME_CHECKED(USubobject, FloatProperty));
 
-			const FPropertySelection* SelectedSubobjectProperties = PropertySelectionMap.GetSelectedProperties(Stub->TestActor->InstancedSubobject);
+			const FPropertySelection* SelectedSubobjectProperties = PropertySelectionMap.GetObjectSelection(Stub->TestActor->InstancedOnlySubobject_DefaultSubobject).GetPropertySelection();
 			const bool bSubobjectHasExpectedNumChangedProperties = SelectedSubobjectProperties && SelectedSubobjectProperties->GetSelectedLeafProperties().Num() == 1;
 			TestTrue(TEXT("Subobject has changed properties"), bSubobjectHasExpectedNumChangedProperties);
 			TestTrue(TEXT("Changed property on subobject contained"), bSubobjectHasExpectedNumChangedProperties && SelectedSubobjectProperties->IsPropertySelected(nullptr, ChangedSubobjectProperty));
 
 
 			// Actor properties
-			const FPropertySelection* SelectedActorProperties = PropertySelectionMap.GetSelectedProperties(Stub->TestActor);
+			const FPropertySelection* SelectedActorProperties = PropertySelectionMap.GetObjectSelection(Stub->TestActor).GetPropertySelection();
 			TestTrue(TEXT("Unchanged actor in selection map when subobject was changed"), SelectedActorProperties != nullptr);
 		});
 
@@ -470,13 +646,13 @@ bool FFilterForPropertiesOnSubobjects::RunTest(const FString& Parameters)
 		.FilterProperties(Stub->TestActor, [&](const FPropertySelectionMap& PropertySelectionMap)
 		{
 			// Custom subobject properties
-			const FPropertySelection* SelectedSubobjectProperties = PropertySelectionMap.GetSelectedProperties(Stub->TestActor->InstancedSubobject);
-			TestTrue(TEXT("Unchanged subobject not in selection map"), SelectedSubobjectProperties == nullptr);
+			const FPropertySelection* SelectedSubobjectProperties = PropertySelectionMap.GetObjectSelection(Stub->TestActor->InstancedOnlySubobject_DefaultSubobject).GetPropertySelection();
+			TestTrue(TEXT("Unchanged subobject not in selection map"), SelectedSubobjectProperties == nullptr || SelectedSubobjectProperties->IsEmpty());
 
 
 			// Actor properties
-			const FPropertySelection* SelectedActorProperties = PropertySelectionMap.GetSelectedProperties(Stub->TestActor);
-			TestTrue(TEXT("Unchanged actor not in selection map when subobject was not changed"), SelectedActorProperties == nullptr);
+			const FPropertySelection* SelectedActorProperties = PropertySelectionMap.GetObjectSelection(Stub->TestActor).GetPropertySelection();
+			TestTrue(TEXT("Unchanged actor not in selection map when subobject was not changed"), SelectedActorProperties == nullptr || SelectedActorProperties->IsEmpty());
 		});
 
 
@@ -491,7 +667,7 @@ bool FFilterForPropertiesOnSubobjects::RunTest(const FString& Parameters)
 		{
 			Stub->TestActor = World->SpawnActor<ASnapshotTestActor>();
 			
-			Stub->TestActor->InstancedSubobject->FloatProperty = 42.f;
+			Stub->TestActor->InstancedOnlySubobject_DefaultSubobject->FloatProperty = 42.f;
 			Stub->TestActor->IntProperty = 42;
 			Stub->TestActor->TestComponent->IntProperty = 42;
 		})
@@ -499,7 +675,7 @@ bool FFilterForPropertiesOnSubobjects::RunTest(const FString& Parameters)
 	
 		.ModifyWorld([&](UWorld* World)
 		{
-			Stub->TestActor->InstancedSubobject->FloatProperty = 420.f;
+			Stub->TestActor->InstancedOnlySubobject_DefaultSubobject->FloatProperty = 420.f;
 			Stub->TestActor->IntProperty = 420;
 			Stub->TestActor->TestComponent->IntProperty = 420;
 		})
@@ -510,7 +686,7 @@ bool FFilterForPropertiesOnSubobjects::RunTest(const FString& Parameters)
 			UClass* SubobjectClass = USubobject::StaticClass();
 			const FProperty* ChangedSubobjectProperty = SubobjectClass->FindPropertyByName(GET_MEMBER_NAME_CHECKED(USubobject, FloatProperty));
 
-			const FPropertySelection* SelectedSubobjectProperties = PropertySelectionMap.GetSelectedProperties(Stub->TestActor->InstancedSubobject);
+			const FPropertySelection* SelectedSubobjectProperties = PropertySelectionMap.GetObjectSelection(Stub->TestActor->InstancedOnlySubobject_DefaultSubobject).GetPropertySelection();
 			const bool bSubobjectHasExpectedNumChangedProperties = SelectedSubobjectProperties && SelectedSubobjectProperties->GetSelectedLeafProperties().Num() == 1;
 			TestTrue(TEXT("Subobject has changed properties"), bSubobjectHasExpectedNumChangedProperties);
 			TestTrue(TEXT("Changed property on subobject contained"), bSubobjectHasExpectedNumChangedProperties && SelectedSubobjectProperties->IsPropertySelected(nullptr, ChangedSubobjectProperty));
@@ -520,7 +696,7 @@ bool FFilterForPropertiesOnSubobjects::RunTest(const FString& Parameters)
 			UClass* TestActorClass = ASnapshotTestActor::StaticClass();
 			const FProperty* ChangedActorProperty = TestActorClass->FindPropertyByName(GET_MEMBER_NAME_CHECKED(ASnapshotTestActor, IntProperty));
 			
-			const FPropertySelection* SelectedActorProperties = PropertySelectionMap.GetSelectedProperties(Stub->TestActor);
+			const FPropertySelection* SelectedActorProperties = PropertySelectionMap.GetObjectSelection(Stub->TestActor).GetPropertySelection();
 			const bool bActorHasExpectedNumChangedProperties = SelectedActorProperties && SelectedActorProperties->GetSelectedLeafProperties().Num();
 			TestTrue(TEXT("Actor has changed properties"), bActorHasExpectedNumChangedProperties);
 			TestTrue(TEXT("Changed property on actor contained"), bActorHasExpectedNumChangedProperties && SelectedActorProperties->IsPropertySelected(nullptr, ChangedActorProperty));
@@ -530,11 +706,118 @@ bool FFilterForPropertiesOnSubobjects::RunTest(const FString& Parameters)
 			UClass* TestComponentClass = USnapshotTestComponent::StaticClass();
 			const FProperty* ChangedComponentProperty = TestComponentClass->FindPropertyByName(GET_MEMBER_NAME_CHECKED(USnapshotTestComponent, IntProperty));
 			
-			const FPropertySelection* SelectedComponentProperties = PropertySelectionMap.GetSelectedProperties(Stub->TestActor->TestComponent);
+			const FPropertySelection* SelectedComponentProperties = PropertySelectionMap.GetObjectSelection(Stub->TestActor->TestComponent).GetPropertySelection();
 			const bool bComponentHasExpectedNumChangedProperties = SelectedComponentProperties && SelectedComponentProperties->GetSelectedLeafProperties().Num();
 			TestTrue(TEXT("Component has changed properties"), bComponentHasExpectedNumChangedProperties);
 			TestTrue(TEXT("Changed property on component contained"), bComponentHasExpectedNumChangedProperties && SelectedComponentProperties->IsPropertySelected(nullptr, ChangedComponentProperty));
 		});
 	
+	return true;
+}
+
+/**
+ * Checks that restoring subobjects which are missing from the editor world are in fact restored.
+ */
+IMPLEMENT_SIMPLE_AUTOMATION_TEST(FRestoreSubobjectsMissingFromEditorWorld, "VirtualProduction.LevelSnapshots.Snapshot.CustomObjectSerialization.RestoreSubobjectsMissingFromEditorWorld", (EAutomationTestFlags::ApplicationContextMask | EAutomationTestFlags::EngineFilter));
+bool FRestoreSubobjectsMissingFromEditorWorld::RunTest(const FString& Parameters)
+{
+	class FStub : public ICustomObjectSnapshotSerializer
+	{
+	public:
+		
+		virtual void OnTakeSnapshot(UObject* EditorObject, ICustomSnapshotSerializationData& DataStorage) override
+		{
+			if (TestActor->NonReflectedSubobject)
+			{
+				DataStorage.AddSubobjectSnapshot(TestActor->NonReflectedSubobject);
+			}
+		}
+
+		virtual UObject* FindOrRecreateSubobjectInSnapshotWorld(UObject* SnapshotObject, const ISnapshotSubobjectMetaData& ObjectData, const ICustomSnapshotSerializationData& DataStorage) override
+		{
+			if (ASnapshotTestActor* SnapshotActor = Cast<ASnapshotTestActor>(SnapshotObject))
+			{
+				SnapshotActor->AllocateNonReflectedSubobject();
+				return SnapshotActor->NonReflectedSubobject;
+			}
+
+			checkNoEntry();
+			return nullptr;
+		}
+		
+		virtual UObject* FindOrRecreateSubobjectInEditorWorld(UObject* EditorObject, const ISnapshotSubobjectMetaData& ObjectData, const ICustomSnapshotSerializationData& DataStorage) override
+		{
+			if (ASnapshotTestActor* EditorActor = Cast<ASnapshotTestActor>(EditorObject))
+			{
+				EditorActor->AllocateNonReflectedSubobject();
+				return EditorActor->NonReflectedSubobject;
+			}
+
+			checkNoEntry();
+			return nullptr;	
+		}
+		
+		virtual UObject* FindSubobjectInEditorWorld(UObject* EditorObject, const ISnapshotSubobjectMetaData& ObjectData, const ICustomSnapshotSerializationData& DataStorage) override
+		{
+			if (ASnapshotTestActor* EditorActor = Cast<ASnapshotTestActor>(EditorObject))
+			{
+				return EditorActor->NonReflectedSubobject;
+			}
+
+			checkNoEntry();
+			return nullptr;
+		}
+
+		ASnapshotTestActor* TestActor;
+	};
+	
+	// Handle registering and unregistering of custom serializer
+	TSharedRef<FStub> Stub = MakeShared<FStub>();
+	FLevelSnapshotsModule& Module = FModuleManager::Get().GetModuleChecked<FLevelSnapshotsModule>("LevelSnapshots");
+	Module.RegisterCustomObjectSerializer(ASnapshotTestActor::StaticClass(), Stub);
+	BlacklistSubobjectProperties();
+	ON_SCOPE_EXIT
+	{
+		Module.UnregisterCustomObjectSerializer(ASnapshotTestActor::StaticClass());
+		UnblacklistSubobjectProperties();
+	};
+
+	
+	FSnapshotTestRunner()
+		.ModifyWorld([&](UWorld* World)
+		{
+			Stub->TestActor = World->SpawnActor<ASnapshotTestActor>();
+			Stub->TestActor->AllocateNonReflectedSubobject();
+			Stub->TestActor->NonReflectedSubobject->FloatProperty = 42.f;
+			Stub->TestActor->NonReflectedSubobject->IntProperty = 21;
+		})
+		.TakeSnapshot()
+
+		.ModifyWorld([&](UWorld* World)
+		{
+			Stub->TestActor->DestroyNonReflectedSubobject();
+		})
+		.FilterProperties(Stub->TestActor, [&](const FPropertySelectionMap& SelectionMap)
+		{
+			const FRestorableObjectSelection ObjectSelection = SelectionMap.GetObjectSelection(Stub->TestActor);
+			TestTrue(TEXT("No changed actor properties"), ObjectSelection.GetPropertySelection() && ObjectSelection.GetPropertySelection()->GetSelectedProperties().Num() == 0 && ObjectSelection.GetPropertySelection()->HasCustomSerializedSubobjects());
+			TestTrue(TEXT("No component selection"), ObjectSelection.GetComponentSelection() == nullptr);
+			TestTrue(TEXT("Needs to restore custom subobject"), ObjectSelection.GetCustomSubobjectSelection() && ObjectSelection.GetCustomSubobjectSelection()->CustomSnapshotSubobjectsToRestore.Num() == 1);
+		})
+		.ApplySnapshot()
+
+		.RunTest([&]()
+		{
+			if (IsValid(Stub->TestActor->NonReflectedSubobject))
+			{
+				TestEqual(TEXT("Custom float restored"), Stub->TestActor->NonReflectedSubobject->FloatProperty, 42.f);
+				TestEqual(TEXT("Custom float restored"), Stub->TestActor->NonReflectedSubobject->IntProperty, 21);
+			}
+			else
+			{
+				AddError(TEXT("Custom subobject was not restored"));
+			}
+		});
+	
 	return true;
 }
\ No newline at end of file
diff --git a/Engine/Plugins/VirtualProduction/LevelSnapshots/Source/LevelSnapshots/Private/Tests/SnapshotTestActor.cpp b/Engine/Plugins/VirtualProduction/LevelSnapshots/Source/LevelSnapshots/Private/Tests/SnapshotTestActor.cpp
--- a/Engine/Plugins/VirtualProduction/LevelSnapshots/Source/LevelSnapshots/Private/Tests/SnapshotTestActor.cpp
+++ b/Engine/Plugins/VirtualProduction/LevelSnapshots/Source/LevelSnapshots/Private/Tests/SnapshotTestActor.cpp
@@ -7,6 +7,7 @@
 USubobject::USubobject()
 {
 	NestedChild = CreateDefaultSubobject<USubSubobject>(TEXT("NestedChild"));
+	UneditableNestedChild = CreateDefaultSubobject<USubSubobject>(TEXT("UneditableNestedChild"));;
 }
 
 USnapshotTestComponent::USnapshotTestComponent()
@@ -48,6 +49,20 @@ ASnapshotTestActor::ASnapshotTestActor()
 	TestComponent = CreateDefaultSubobject<USnapshotTestComponent>(TEXT("TestComponent"));
 }
 
+bool ASnapshotTestActor::HasAnyValidHardObjectReference() const
+{
+	const bool bHasSingleReference = !ObjectReference || !WeakObjectPtr.IsValid();
+
+	const bool bAnyCollectionHasReference = ObjectArray.Num() != 0
+		|| ObjectSet.Num() != 0
+		|| WeakObjectPtrArray.Num() != 0
+		|| WeakObjectPtrSet.Num() != 0
+		|| ObjectMap.Num() != 0
+		|| WeakObjectPtrMap.Num() != 0;
+
+	return !bHasSingleReference && !bAnyCollectionHasReference;
+}
+
 bool ASnapshotTestActor::HasObjectReference(UObject* Object, bool bOnlyCheckCollections, FName MapKey) const
 {
 	const bool bReferencesEqual = (bOnlyCheckCollections
@@ -121,11 +136,45 @@ void ASnapshotTestActor::ClearObjectReferences()
 	WeakObjectPtrMap.Reset();
 }
 
+void ASnapshotTestActor::AllocateSubobjects()
+{
+	EditOnlySubobject_OptionalSubobject = NewObject<USubobject>(this, USubobject::StaticClass(), TEXT("EditOnlySubobject"));
+
+	EditableInstancedSubobjectArray_OptionalSubobject.Add(NewObject<USubobject>(this, USubobject::StaticClass(), TEXT("EditableInstancedSubobjectArray_1")));
+	EditableInstancedSubobjectArray_OptionalSubobject.Add(NewObject<USubobject>(this, USubobject::StaticClass(), TEXT("EditableInstancedSubobjectArray_2")));
+
+	EditableInstancedSubobjectMap_OptionalSubobject.Add("First", NewObject<USubobject>(this, USubobject::StaticClass(), TEXT("EditableInstancedSubobjectMap_1")));
+	EditableInstancedSubobjectMap_OptionalSubobject.Add("Second", NewObject<USubobject>(this, USubobject::StaticClass(), TEXT("EditableInstancedSubobjectMap_2")));
+	
+	EditOnlySubobjectArray_OptionalSubobject.Add(NewObject<USubobject>(this, USubobject::StaticClass(), TEXT("EditOnlySubobjectArray_1")));
+	EditOnlySubobjectArray_OptionalSubobject.Add(NewObject<USubobject>(this, USubobject::StaticClass(), TEXT("EditOnlySubobjectArray_2")));
+	
+	EditOnlySubobjectMap_OptionalSubobject.Add("First", NewObject<USubobject>(this, USubobject::StaticClass(), TEXT("EditOnlySubobjectMap_1")));
+	EditOnlySubobjectMap_OptionalSubobject.Add("Second", NewObject<USubobject>(this, USubobject::StaticClass(), TEXT("EditOnlySubobjectMap_2")));
+}
+
+void ASnapshotTestActor::AllocateNonReflectedSubobject()
+{
+	if (!NonReflectedSubobject)
+	{
+		NonReflectedSubobject = NewObject<USubobject>(this, USubobject::StaticClass(), TEXT("NonReflectedSubobject"));
+	}
+}
+
+void ASnapshotTestActor::DestroyNonReflectedSubobject()
+{
+	if (NonReflectedSubobject)
+	{
+		NonReflectedSubobject->MarkPendingKill();
+		NonReflectedSubobject = nullptr;
+	}
+}
+
 void ASnapshotTestActor::PostInitProperties()
 {
 	Super::PostInitProperties();
 
-	EditableInstancedSubobject = NewObject<USubobject>(this, USubobject::StaticClass(), TEXT("EditableInstancedSubobject"));
-	InstancedSubobject = NewObject<USubobject>(this, USubobject::StaticClass(), TEXT("InstancedSubobject"));
-	NakedSubobject = NewObject<USubobject>(this, USubobject::StaticClass(), TEXT("NakedSubobject"));;
+	EditableInstancedSubobject_DefaultSubobject = NewObject<USubobject>(this, USubobject::StaticClass(), TEXT("EditableInstancedSubobject"));
+	InstancedOnlySubobject_DefaultSubobject = NewObject<USubobject>(this, USubobject::StaticClass(), TEXT("InstancedSubobject"));
+	NakedSubobject_DefaultSubobject = NewObject<USubobject>(this, USubobject::StaticClass(), TEXT("NakedSubobject"));
 }
