Refresh UE5 branch with latest Level Snapshots which includes support for subobjects.

diff --git a/Engine/Plugins/VirtualProduction/LevelSnapshots/Source/LevelSnapshots/Private/Tests/SnapshotTestActor.h b/Engine/Plugins/VirtualProduction/LevelSnapshots/Source/LevelSnapshots/Private/Tests/SnapshotTestActor.h
--- a/Engine/Plugins/VirtualProduction/LevelSnapshots/Source/LevelSnapshots/Private/Tests/SnapshotTestActor.h
+++ b/Engine/Plugins/VirtualProduction/LevelSnapshots/Source/LevelSnapshots/Private/Tests/SnapshotTestActor.h
@@ -39,8 +39,15 @@ public:
 	UPROPERTY(EditAnywhere, Category = "Level Snapshots")
 	float FloatProperty;
 
-	UPROPERTY(Instanced)
+	UPROPERTY(EditAnywhere, Instanced, Category = "Level Snapshots")
 	USubSubobject* NestedChild;
+
+	UPROPERTY(Instanced)
+	USubSubobject* UneditableNestedChild;
+
+	FName NonReflectedName;
+	UObject* NonReflectedObjectProperty;
+	TSoftObjectPtr<UObject> NonReflectedSoftPtr;
 };
 
 UCLASS()
@@ -69,12 +76,18 @@ public:
 
 	ASnapshotTestActor();
 
+	bool HasAnyValidHardObjectReference() const;
 	bool HasObjectReference(UObject* Object, bool bOnlyCheckCollections = false, FName MapKey = NAME_Name) const;
 	
 	void SetObjectReference(UObject* Object, FName MapKey = NAME_Name);
 	void AddObjectReference(UObject* Object, FName MapKey = NAME_Name);
 	void ClearObjectReferences();
 
+	void AllocateSubobjects();
+
+	void AllocateNonReflectedSubobject();
+	void DestroyNonReflectedSubobject();
+	
 	//~ Begin UObject Interface
 	virtual void PostInitProperties() override;
 	//~ End UObject Interface
@@ -202,11 +215,33 @@ public:
 	/******************** Subobject references  ********************/
 
 	UPROPERTY(EditAnywhere, Instanced, Category = "Level Snapshots")
-	USubobject* EditableInstancedSubobject;
+	USubobject* EditableInstancedSubobject_DefaultSubobject;
 
 	UPROPERTY(Instanced)
-	USubobject* InstancedSubobject;
+	USubobject* InstancedOnlySubobject_DefaultSubobject;
 
 	UPROPERTY()
-	USubobject* NakedSubobject;
+	USubobject* NakedSubobject_DefaultSubobject;
+
+	UPROPERTY(EditAnywhere, Category = "Level Snapshots")
+	USubobject* EditOnlySubobject_OptionalSubobject;
+	
+	
+	UPROPERTY(EditAnywhere, Instanced, Category = "Level Snapshots")
+	TArray<USubobject*> EditableInstancedSubobjectArray_OptionalSubobject;
+
+	UPROPERTY(EditAnywhere, Instanced, Category = "Level Snapshots")
+	TMap<FName, USubobject*> EditableInstancedSubobjectMap_OptionalSubobject;
+
+	UPROPERTY(EditAnywhere, Category = "Level Snapshots")
+	TArray<USubobject*> EditOnlySubobjectArray_OptionalSubobject;
+
+	UPROPERTY(EditAnywhere, Category = "Level Snapshots")
+	TMap<FName, USubobject*> EditOnlySubobjectMap_OptionalSubobject;
+	
+	USubobject* NonReflectedSubobject;
+
+	FName NonReflectedName;
+	UObject* NonReflectedObjectProperty;
+	TSoftObjectPtr<UObject> NonReflectedSoftPtr;
 };
diff --git a/Engine/Plugins/VirtualProduction/LevelSnapshots/Source/LevelSnapshots/Public/Data/ActorComponentSelection.h b/Engine/Plugins/VirtualProduction/LevelSnapshots/Source/LevelSnapshots/Public/Data/ActorComponentSelection.h
new file mode 100644
--- /dev/null
+++ b/Engine/Plugins/VirtualProduction/LevelSnapshots/Source/LevelSnapshots/Public/Data/ActorComponentSelection.h
@@ -0,0 +1,28 @@
+﻿// Copyright Epic Games, Inc. All Rights Reserved.
+
+#pragma once
+
+#include "CoreMinimal.h"
+#include "UObject/WeakObjectPtrTemplates.h"
+
+class UActorComponent;
+
+/* Holds all properties that should be restored for an object. */
+struct LEVELSNAPSHOTS_API FActorComponentSelection
+{
+	/**
+	 * New components added to the actor after the snapshot was taken. Remove them from the original actor.
+	 * Point to instances in the editor world.
+	 */
+	TSet<TWeakObjectPtr<UActorComponent>> EditorWorldComponentsToRemove;
+
+	/**
+	 * Components that were removed from the actor after the snapshot was taken. Add them to the original actor.
+	 * Point to instances in the snapshot world.
+	 */
+	TSet<TWeakObjectPtr<UActorComponent>> SnapshotComponentsToAdd;
+};
+		
+		
+
+		
\ No newline at end of file
diff --git a/Engine/Plugins/VirtualProduction/LevelSnapshots/Source/LevelSnapshots/Public/Data/ActorSnapshotData.h b/Engine/Plugins/VirtualProduction/LevelSnapshots/Source/LevelSnapshots/Public/Data/ActorSnapshotData.h
--- a/Engine/Plugins/VirtualProduction/LevelSnapshots/Source/LevelSnapshots/Public/Data/ActorSnapshotData.h
+++ b/Engine/Plugins/VirtualProduction/LevelSnapshots/Source/LevelSnapshots/Public/Data/ActorSnapshotData.h
@@ -6,6 +6,8 @@
 #include "ComponentSnapshotData.h"
 #include "CustomSerializationData.h"
 #include "ObjectSnapshotData.h"
+#include "SubobjectSnapshotData.h"
+
 #include "ActorSnapshotData.generated.h"
 
 class UActorComponent;
@@ -14,8 +16,6 @@ class UPackage;
 struct FPropertySelectionMap;
 struct FWorldSnapshotData;
 
-// TODO: Move these functions into static helper functions in private part of module.
-
 USTRUCT()
 struct LEVELSNAPSHOTS_API FActorSnapshotData
 {
@@ -23,24 +23,30 @@ struct LEVELSNAPSHOTS_API FActorSnapshotData
 
 	static FActorSnapshotData SnapshotActor(AActor* OriginalActor, FWorldSnapshotData& WorldData);
 	
-	void DeserializeIntoExistingWorldActor(UWorld* SnapshotWorld, AActor* OriginalActor, FWorldSnapshotData& WorldData, UPackage* InLocalisationSnapshotPackage, const FPropertySelectionMap& SelectedProperties);
-	void DeserializeIntoRecreatedEditorWorldActor(UWorld* SnapshotWorld, AActor* OriginalActor, FWorldSnapshotData& WorldData, UPackage* InLocalisationSnapshotPackage, const FPropertySelectionMap& SelectedProperties);
-
 	TOptional<AActor*> GetPreallocatedIfValidButDoNotAllocate() const;
 	TOptional<AActor*> GetPreallocated(UWorld* SnapshotWorld, FWorldSnapshotData& WorldData) const;
-	TOptional<AActor*> GetDeserialized(UWorld* SnapshotWorld, FWorldSnapshotData& WorldData, UPackage* InLocalisationSnapshotPackage);
+	TOptional<AActor*> GetDeserialized(UWorld* SnapshotWorld, FWorldSnapshotData& WorldData, const FSoftObjectPath& OriginalActorPath, UPackage* InLocalisationSnapshotPackage);
 
+	/**
+	 * Fixes up any properties that were not yet captured
+	 */
+	void PostSerializeSnapshotActor(AActor* SnapshotActor, FWorldSnapshotData& WorldData, const FSoftObjectPath& OriginalActorPath, UPackage* InLocalisationSnapshotPackage) const;
+
+	void DeserializeIntoExistingWorldActor(UWorld* SnapshotWorld, AActor* OriginalActor, FWorldSnapshotData& WorldData, UPackage* InLocalisationSnapshotPackage, const FPropertySelectionMap& SelectedProperties);
+	void DeserializeIntoRecreatedEditorWorldActor(UWorld* SnapshotWorld, AActor* OriginalActor, FWorldSnapshotData& WorldData, UPackage* InLocalisationSnapshotPackage, const FPropertySelectionMap& SelectedProperties);
+	
 	const FSoftClassPath& GetActorClass() const { return ActorClass; }
 	const FCustomSerializationData& GetCustomActorSerializationData() const { return CustomActorSerializationData; }
 
-private:
-
+	
 	using FSerializeActor = TFunction<void(AActor* OriginalActor, AActor* DerserializedActor)>;
-	using FSerializeComponent = TFunction<void(FObjectSnapshotData& SerializedCompData, FComponentSnapshotData& CompData, UActorComponent* Original, UActorComponent* Deserialized)>;
-	using FHandleFoundComponent = TFunction<void(FObjectSnapshotData& SerializedCompData, FComponentSnapshotData& ComponentMetaData, UActorComponent* ActorComp, const FSoftObjectPath& OriginalComponentPath, FWorldSnapshotData& SharedData)>;
+	using FSerializeComponent = TFunction<void(FSubobjectSnapshotData& SerializedCompData, FComponentSnapshotData& CompData, UActorComponent* Original, UActorComponent* Deserialized)>;
+	using FHandleFoundComponent = TFunction<void(FSubobjectSnapshotData& SerializedCompData, FComponentSnapshotData& ComponentMetaData, UActorComponent* ActorComp, const FSoftObjectPath& OriginalComponentPath, FWorldSnapshotData& SharedData)>;
 	
-	void DeserializeIntoWorldActor(UWorld* SnapshotWorld, AActor* OriginalActor, FWorldSnapshotData& WorldData, UPackage* InLocalisationSnapshotPackage, FSerializeActor SerializeActor, FSerializeComponent SerializeComponent);
+	void DeserializeIntoEditorWorldActor(UWorld* SnapshotWorld, AActor* OriginalActor, FWorldSnapshotData& WorldData, UPackage* InLocalisationSnapshotPackage, FSerializeActor SerializeActor, FSerializeComponent SerializeComponent);
 	void DeserializeComponents(AActor* IntoActor, FWorldSnapshotData& WorldData, FHandleFoundComponent Callback);
+
+	void DeserializeSubobjectsForSnapshotActor(AActor* IntoActor, FWorldSnapshotData& WorldData, UPackage* InLocalisationSnapshotPackage);
 	
 	/* We cache the actor to avoid recreating it all the time */
 	UPROPERTY(Transient)
@@ -54,6 +60,8 @@ private:
 	UPROPERTY(Transient)
 	mutable bool bReceivedSerialisation = false;
 
+	
+
 	/** The class the actor had when it was saved. */
 	UPROPERTY()
 	FSoftClassPath ActorClass;
@@ -61,12 +69,20 @@ private:
 	/** The actor's serialized data */
 	UPROPERTY()
 	FObjectSnapshotData SerializedActorData;
+	
 
+	
 	/** Data that was generated by some ICustomObjectSnapshotSerializer. Needed to restore custom subobjects. */
 	UPROPERTY()
 	FCustomSerializationData CustomActorSerializationData;
+
+
 	
 	/** Additional component recreation info */
 	UPROPERTY()
 	TMap<int32, FComponentSnapshotData> ComponentData;
-};
\ No newline at end of file
+
+	/** Tracks all non-component subobjects. Valid index to FWorldSnapshotData::SerializedObjectReferences. */
+	UPROPERTY()
+	TArray<int32> OwnedSubobjects;
+};
diff --git a/Engine/Plugins/VirtualProduction/LevelSnapshots/Source/LevelSnapshots/Public/Data/ComponentSnapshotData.h b/Engine/Plugins/VirtualProduction/LevelSnapshots/Source/LevelSnapshots/Public/Data/ComponentSnapshotData.h
--- a/Engine/Plugins/VirtualProduction/LevelSnapshots/Source/LevelSnapshots/Public/Data/ComponentSnapshotData.h
+++ b/Engine/Plugins/VirtualProduction/LevelSnapshots/Source/LevelSnapshots/Public/Data/ComponentSnapshotData.h
@@ -18,11 +18,9 @@ struct LEVELSNAPSHOTS_API FComponentSnapshotData
 
 	static TOptional<FComponentSnapshotData> SnapshotComponent(UActorComponent* OriginalComponent, FWorldSnapshotData& WorldData);
 	
-	bool IsRestoreSupportedForSavedComponent() const;
-	
 	void DeserializeIntoTransient(FObjectSnapshotData& SerializedComponentData, UActorComponent* ComponentToDeserializeInto, FWorldSnapshotData& WorldData, UPackage* InLocalisationSnapshotPackage);
+	
 	/* Describes how the component was created */
 	UPROPERTY()
 	EComponentCreationMethod CreationMethod;
-	
 };
\ No newline at end of file
diff --git a/Engine/Plugins/VirtualProduction/LevelSnapshots/Source/LevelSnapshots/Public/Data/PropertySelection.h b/Engine/Plugins/VirtualProduction/LevelSnapshots/Source/LevelSnapshots/Public/Data/PropertySelection.h
--- a/Engine/Plugins/VirtualProduction/LevelSnapshots/Source/LevelSnapshots/Public/Data/PropertySelection.h
+++ b/Engine/Plugins/VirtualProduction/LevelSnapshots/Source/LevelSnapshots/Public/Data/PropertySelection.h
@@ -19,6 +19,9 @@ struct LEVELSNAPSHOTS_API FLevelSnapshotPropertyChain : FArchiveSerializedProper
 	using Super::GetPropertyFromStack;
 	using Super::GetPropertyFromRoot;
 	using Super::GetNumProperties;
+
+	FLevelSnapshotPropertyChain() = default;
+	FLevelSnapshotPropertyChain(const FProperty* RootProperty);
 	
 	FLevelSnapshotPropertyChain MakeAppended(const FProperty* Property) const;
 	void AppendInline(const FProperty* Property);
@@ -38,6 +41,11 @@ struct LEVELSNAPSHOTS_API FLevelSnapshotPropertyChain : FArchiveSerializedProper
 /* Holds all properties that should be restored for an object. */
 struct LEVELSNAPSHOTS_API FPropertySelection
 {
+	FPropertySelection() = default;
+	FPropertySelection(const FProperty* SingleProperty);
+	/* Note: all properties must be root properties (i.e. not within structs or collections). Check yourself, this is not checked. */
+	FPropertySelection(const TSet<const FProperty*>& Properties);
+	
 	/** 
 	 * Checks whether the given property should be serialized. It should be serialized if:
 	 * - IsPropertySelected() returns true on the property
@@ -61,8 +69,8 @@ struct LEVELSNAPSHOTS_API FPropertySelection
 
 	bool IsEmpty() const;
 
-	void SetHasCustomSerializedSubobjects(bool bValue) { bHasCustomSerializedSubobjects = bValue; };
-	bool HasCustomSerializedSubobjects() const { return bHasCustomSerializedSubobjects;}
+	void SetHasCustomSerializedSubobjects(bool bValue) { bHasCustomSerializedSubobjects = bValue; }
+	bool HasCustomSerializedSubobjects() const { return bHasCustomSerializedSubobjects; }
 
 	void AddProperty(const FLevelSnapshotPropertyChain& SelectedProperty);
 	void RemoveProperty(const FArchiveSerializedPropertyChain* ContainerChain, const FProperty* LeafProperty);
diff --git a/Engine/Plugins/VirtualProduction/LevelSnapshots/Source/LevelSnapshots/Public/Data/SnapshotVersion.h b/Engine/Plugins/VirtualProduction/LevelSnapshots/Source/LevelSnapshots/Public/Data/SnapshotVersion.h
--- a/Engine/Plugins/VirtualProduction/LevelSnapshots/Source/LevelSnapshots/Public/Data/SnapshotVersion.h
+++ b/Engine/Plugins/VirtualProduction/LevelSnapshots/Source/LevelSnapshots/Public/Data/SnapshotVersion.h
@@ -82,11 +82,14 @@ struct LEVELSNAPSHOTS_API FSnapshotVersionInfo
 	GENERATED_BODY()
 
 	/** Initialize this version info from the compiled in data */
-	void Initialize();
+	void Initialize(bool bWithoutSnapshotVersion = false);
 	bool IsInitialized() const;
 
 	void ApplyToArchive(FArchive& Archive) const;
 	
+	/** @return The saved version for this plugin or -1 if none is available. */
+	int32 GetSnapshotCustomVersion() const;
+	
 	/** File version info */
 	UPROPERTY()
 	FSnapshotFileVersionInfo FileVersion;
diff --git a/Engine/Plugins/VirtualProduction/LevelSnapshots/Source/LevelSnapshots/Public/Data/SubobjectSnapshotData.h b/Engine/Plugins/VirtualProduction/LevelSnapshots/Source/LevelSnapshots/Public/Data/SubobjectSnapshotData.h
--- a/Engine/Plugins/VirtualProduction/LevelSnapshots/Source/LevelSnapshots/Public/Data/SubobjectSnapshotData.h
+++ b/Engine/Plugins/VirtualProduction/LevelSnapshots/Source/LevelSnapshots/Public/Data/SubobjectSnapshotData.h
@@ -11,11 +11,28 @@ struct LEVELSNAPSHOTS_API FSubobjectSnapshotData : public FObjectSnapshotData
 {
 	GENERATED_BODY()
 
+	static FSubobjectSnapshotData MakeBlacklisted()
+	{
+		FSubobjectSnapshotData Result;
+		Result.bWasBlacklisted = true;
+		return Result;
+	}
+
 	/* Index to FWorldSnapshotData::SerializedObjectReferences */
 	UPROPERTY()
-	int32 OuterIndex;
+	int32 OuterIndex = INDEX_NONE;
 
 	UPROPERTY()
 	FSoftClassPath Class;
-	
+
+	UPROPERTY()
+	bool bWasBlacklisted = false;
+
+	/** Allocated in snapshot world */
+	UPROPERTY(Transient)
+	TWeakObjectPtr<UObject> SnapshotObject;
+
+	/** Allocated in editor world */
+	UPROPERTY(Transient)
+	TWeakObjectPtr<UObject> EditorObject;
 };
\ No newline at end of file
diff --git a/Engine/Plugins/VirtualProduction/LevelSnapshots/Source/LevelSnapshots/Public/Data/WorldSnapshotData.h b/Engine/Plugins/VirtualProduction/LevelSnapshots/Source/LevelSnapshots/Public/Data/WorldSnapshotData.h
--- a/Engine/Plugins/VirtualProduction/LevelSnapshots/Source/LevelSnapshots/Public/Data/WorldSnapshotData.h
+++ b/Engine/Plugins/VirtualProduction/LevelSnapshots/Source/LevelSnapshots/Public/Data/WorldSnapshotData.h
@@ -17,8 +17,6 @@ class FSnapshotArchive;
 struct FPropertySelectionMap;
 struct FActorSnapshotData;
 
-// TODO: Move all functions into static helper functions in private part of module. Helps separation of concerns.
-
 /* Holds saved world data and handles all logic related to writing to the existing world. */
 USTRUCT()
 struct LEVELSNAPSHOTS_API FWorldSnapshotData
@@ -39,43 +37,15 @@ struct LEVELSNAPSHOTS_API FWorldSnapshotData
 	void ForEachOriginalActor(TFunction<void(const FSoftObjectPath& ActorPath, const FActorSnapshotData& SavedData)> HandleOriginalActorPath) const;
 	bool HasMatchingSavedActor(const FSoftObjectPath& OriginalObjectPath) const;
 
-	/* Given a path to a world actor, gets the equivalent allocated actor. */
+	/** Allocates a snapshot actor for the given editor world actor but does not serialize any data into it. */
+	TOptional<AActor*> GetPreallocatedActor(const FSoftObjectPath& OriginalObjectPath);
+	/** Same as GetPreallocatedActor, only that all data will be serialized into it. */
 	TOptional<AActor*> GetDeserializedActor(const FSoftObjectPath& OriginalObjectPath, UPackage* LocalisationSnapshotPackage);
 	/* Gets the state of the CDO from when the snapshot was taken. CDO is saved when a snapshot is taken so CDO changes can be detected. */
 	TOptional<FObjectSnapshotData*> GetSerializedClassDefaults(UClass* Class);  
-
-	/**
-	 * Checks whether two pointers point to "equivalent" objects.
-	 */
-	bool AreReferencesEquivalent(UObject* SnapshotPropertyValue, UObject* OriginalPropertyValue, AActor* SnapshotActor, AActor* OriginalActor) const;
 	
 public: /* Serialisation functions */
 	
-	
-	/* Adds an object dependency without serializing the object's content. Intended for external objects, e.g. to UMaterial in the content browser. */
-	int32 AddObjectDependency(UObject* ReferenceFromOriginalObject);
-
-	/* Resolves an object dependency for use in the snapshot world. If the object is a subobject, it gets fully serialized. */
-	UObject* ResolveObjectDependencyForSnapshotWorld(int32 ObjectPathIndex);
-	
-	/* Resolves an object dependency for use in the editor world. If the object is a subobject, it is serialized.
-	 * Steps for serializing subobject:
-	 *  - If an equivalent object with the saved name and class exists, use that and serialize the properties in SelectionMap into it.
-	 *  - Otherwise allocate a new object and serialize all properties into it
-	 *  - Replaces all references to the trashed object with the new one
-	 */
-	UObject* ResolveObjectDependencyForEditorWorld(int32 ObjectPathIndex, const FPropertySelectionMap& SelectionMap);
-
-	/* Resolves an object depedency when restoring a class default object. Simply resolves without further checks. */
-	UObject* ResolveObjectDependencyForClassDefaultObject(int32 ObjectPathIndex);
-	
-	
-	/**
-	 * Adds a subobject dependency. Intended for internal objects which need to store serialized data, e.g. components and other subobjects. Implicitly calls AddObjectDependency.
-	 * @return A valid index in SerializedObjectReferences and the corresponding subobject data.
-	 */
-	int32 AddSubobjectDependency(UObject* ReferenceFromOriginalObject);
-
 	/**
 	 * Adds a subobject to SerializedObjectReferences and CustomSubobjectSerializationData.
 	 * @return A valid index in SerializedObjectReferences and the corresponding subobject data.
@@ -101,8 +71,6 @@ public: /* Serialisation functions */
 	//~ End TStructOpsTypeTraits Interface
 	
 public:
-
-	UObject* ResolveExternalReference(const FSoftObjectPath& ObjectPath);
 	
 	void ApplyToWorld_HandleRemovingActors(UWorld* WorldToApplyTo, const FPropertySelectionMap& PropertiesToSerialize);
 	void ApplyToWorld_HandleRecreatingActors(TSet<AActor*>& EvaluatedActors, UPackage* LocalisationSnapshotPackage, const FPropertySelectionMap& PropertiesToSerialize);
@@ -149,7 +117,7 @@ public:
 	 * Example: UStaticMesh /Game/Australia/StaticMeshes/MegaScans/Nature_Rock_vbhtdixga/vbhtdixga_LOD0.vbhtdixga_LOD0
 	 * 
 	 * Internal references, e.g. to subobjects and to other actors in the world, are a bit tricky.
-	 * For internal references, we need to do some translation using TranslateOriginalToSnapshotPath:
+	 * For internal references, we need to do some translation:
 	 * Example original: UStaticMeshActor::StaticMeshComponent /Game/MapName.MapName:PersistentLevel.StaticMeshActor_42.StaticMeshComponent
 	 * Example translated: UStaticMeshActor::StaticMeshComponent /Engine/Transient.World_21:PersistentLevel.StaticMeshActor_42.StaticMeshComponent
 	 */
diff --git a/Engine/Plugins/VirtualProduction/LevelSnapshots/Source/LevelSnapshots/Public/ILevelSnapshotsModule.h b/Engine/Plugins/VirtualProduction/LevelSnapshots/Source/LevelSnapshots/Public/ILevelSnapshotsModule.h
--- a/Engine/Plugins/VirtualProduction/LevelSnapshots/Source/LevelSnapshots/Public/ILevelSnapshotsModule.h
+++ b/Engine/Plugins/VirtualProduction/LevelSnapshots/Source/LevelSnapshots/Public/ILevelSnapshotsModule.h
@@ -3,6 +3,7 @@
 #pragma once
 
 #include "CoreMinimal.h"
+#include "IRestorationListener.h"
 
 #include "Restorability/IPropertyComparer.h"
 #include "Restorability/ISnapshotRestorabilityOverrider.h"
@@ -41,6 +42,7 @@ public:
 	/** Called after a snapshot is taken */
 	FPostTakeSnapshotEvent& OnPostTakeSnapshot() { return PostTakeSnapshot; }
 
+	
 	DECLARE_DELEGATE_RetVal_OneParam(bool, FCanTakeSnapshot, const FPreTakeSnapshotEventData&);
 	/** Add a named delegate that determines if we can take a snapshot. */
 	virtual void AddCanTakeSnapshotDelegate(FName DelegateName, FCanTakeSnapshot Delegate) = 0;
@@ -51,7 +53,12 @@ public:
 	/** Queries the attached snapshot delegate and determines if we can take a snapshot.*/
 	virtual bool CanTakeSnapshot(const FPreTakeSnapshotEventData& Event) const = 0;
 
-	/** */
+
+	/** Snapshots will no longer capture nor restore subobjects of this class. Subclasses are implicitly blacklisted. */
+	virtual void AddBlacklistedSubobjectClasses(const TSet<UClass*>& Classes) = 0;
+	virtual void RemoveBlacklistedSubobjectClasses(const TSet<UClass*>& Classes) = 0;
+	
+	
 	/* Registers callbacks that override which actors, components, and properties are restored by default. */
 	virtual void RegisterRestorabilityOverrider(TSharedRef<ISnapshotRestorabilityOverrider> Overrider) = 0;
 	/* Unregisters an overrider previously registered. */
@@ -73,6 +80,11 @@ public:
 	 */
 	virtual void RegisterCustomObjectSerializer(UClass* Class, TSharedRef<ICustomObjectSnapshotSerializer> CustomSerializer, bool bIncludeBlueprintChildClasses = true) = 0;
 	virtual void UnregisterCustomObjectSerializer(UClass* Class) = 0;
+
+	
+	/** Registers an object that will receive callbacks when a snapshot is applied. */
+	virtual void RegisterRestorationListener(TSharedRef<IRestorationListener> Listener) = 0;
+	virtual void UnregisterRestorationListener(TSharedRef<IRestorationListener> Listener) = 0;
 	
 	
 	/**
@@ -82,7 +94,8 @@ public:
 	virtual void AddWhitelistedProperties(const TSet<const FProperty*>& Properties) = 0;
 	virtual void RemoveWhitelistedProperties(const TSet<const FProperty*>& Properties) = 0;
 
-	/** Stops snapshots from capturing / restoring these properties.
+	/**
+	 * Stops snapshots from capturing / restoring these properties.
 	 * Important: Only add add native properties; Blueprint properties may be invalidated (and left dangeling) when recompiled.
 	 */
 	virtual void AddBlacklistedProperties(const TSet<const FProperty*>& Properties) = 0;
diff --git a/Engine/Plugins/VirtualProduction/LevelSnapshots/Source/LevelSnapshots/Public/LevelSnapshotsFunctionLibrary.h b/Engine/Plugins/VirtualProduction/LevelSnapshots/Source/LevelSnapshots/Public/LevelSnapshotsFunctionLibrary.h
--- a/Engine/Plugins/VirtualProduction/LevelSnapshots/Source/LevelSnapshots/Public/LevelSnapshotsFunctionLibrary.h
+++ b/Engine/Plugins/VirtualProduction/LevelSnapshots/Source/LevelSnapshots/Public/LevelSnapshotsFunctionLibrary.h
@@ -40,12 +40,19 @@ public:
 		);
 
 	/**
-	 * Recursively gets all subobjects which should be serialized using custom callbacks from an external module.
+	 * Recursively gets all subobjects which should be serialized using an ICustomObjectSnapshotSerializer implementation from an external module.
 	 *
 	 * @param Snapshot The snapshot in which to look
 	 * @param SnapshotRootObject Snapshot version of the object to look in. Either an actor, component, or a custom serialized subobject.
 	 * @param WorldRootObject Editor version of the object to look in. Either an actor, component, or a custom serialized subobject.
-	 * @param Callback Whatever you want to do with this information
+	 * @param HandleCustomSubobjectPair Called when two custom subobjects are matched
+	 * @param HandleUnmatchedSnapshotSubobject Called when a custom snapshot subobject exists but no corresponding editor subobject exists
 	 */
-	static void ForEachMatchingCustomSubobjectPair(ULevelSnapshot* Snapshot, UObject* SnapshotRootObject, UObject* WorldRootObject, TFunction<void(UObject* SnapshotSubobject, UObject* EditorWorldSubobject)> Callback);
+	static void ForEachMatchingCustomSubobjectPair(
+		ULevelSnapshot* Snapshot, 
+		UObject* SnapshotRootObject,
+		UObject* WorldRootObject,
+		TFunction<void(UObject* SnapshotSubobject, UObject* EditorWorldSubobject)> HandleCustomSubobjectPair,
+		TFunction<void(UObject* UnmatchedSnapshotSubobject)> HandleUnmatchedSnapshotSubobject
+		);
 };
\ No newline at end of file
diff --git a/Engine/Plugins/VirtualProduction/LevelSnapshots/Source/LevelSnapshots/Public/PropertyInfoHelpers.h b/Engine/Plugins/VirtualProduction/LevelSnapshots/Source/LevelSnapshots/Public/PropertyInfoHelpers.h
--- a/Engine/Plugins/VirtualProduction/LevelSnapshots/Source/LevelSnapshots/Public/PropertyInfoHelpers.h
+++ b/Engine/Plugins/VirtualProduction/LevelSnapshots/Source/LevelSnapshots/Public/PropertyInfoHelpers.h
@@ -11,8 +11,10 @@ class UObject;
 
 class LEVELSNAPSHOTS_API FPropertyInfoHelpers
 {
+	static float FloatComparisonPrecision;
+	static double DoubleComparisonPrecision;
 public:
-
+	
 	static FProperty* GetParentProperty(const FProperty* Property);
 
 	/* Returns true if the property is a struct, map, array, set, etc */
@@ -35,6 +37,7 @@ public:
 	/* A quick property flag check. Assumes the property flags are properly set/deserialized */
 	static bool IsPropertySubObject(const FProperty* Property);
 
+	static void UpdateDecimalComparisionPrecision(float FloatPrecision, double DoublePrecision);
 	static bool AreNumericPropertiesNearlyEqual(const FNumericProperty* Property, const void* ValuePtrA, const void* ValuePtrB);
 	
 };
diff --git a/Engine/Plugins/VirtualProduction/LevelSnapshots/Source/LevelSnapshots/Public/PropertySelectionMap.h b/Engine/Plugins/VirtualProduction/LevelSnapshots/Source/LevelSnapshots/Public/PropertySelectionMap.h
--- a/Engine/Plugins/VirtualProduction/LevelSnapshots/Source/LevelSnapshots/Public/PropertySelectionMap.h
+++ b/Engine/Plugins/VirtualProduction/LevelSnapshots/Source/LevelSnapshots/Public/PropertySelectionMap.h
@@ -3,17 +3,24 @@
 #pragma once
 
 #include "CoreMinimal.h"
-#include "PropertySelection.h"
+#include "Data/AddedAndRemovedComponentInfo.h"
+#include "Data/CustomSubobjectRestorationInfo.h"
+#include "Data/PropertySelection.h"
+#include "Data/RestorableObjectSelection.h"
 #include "PropertySelectionMap.generated.h"
 
 class AActor;
+struct FPropertySelectionMap;
 
 /* Binds an object to its selected properties */
 USTRUCT()
 struct LEVELSNAPSHOTS_API FPropertySelectionMap
 {
 	GENERATED_BODY()
+	friend FRestorableObjectSelection;
 
+	/*************** Common features ***************/
+	
 	/* Respawn the actor from the data in the snapshot. */
 	void AddDeletedActorToRespawn(const FSoftObjectPath& Original);
 	void RemoveDeletedActorToRespawn(const FSoftObjectPath& Original);
@@ -22,15 +29,49 @@ struct LEVELSNAPSHOTS_API FPropertySelectionMap
 	void AddNewActorToDespawn(AActor* WorldActor);
 	void RemoveNewActorToDespawn(AActor* WorldActor);
 
-	/* Binds properties to an object which are supposed to be rolled back. */
+	/**
+	 * Binds properties to an object which are supposed to be rolled back.
+	 */
 	bool AddObjectProperties(UObject* WorldObject, const FPropertySelection& SelectedProperties);
 	void RemoveObjectPropertiesFromMap(UObject* WorldObject);
+
 	
-	const FPropertySelection* GetSelectedProperties(UObject* WorldObject) const;
+	/*************** Basic subobject support ***************/
+	
+	/**
+	 * Tracks a subobject so we only restore references pointing to it but not serialize its properties.
+	 *
+	 * When applying a snapshot to world and resolving subobjects, it is unclear whether a resolved references is either
+	 *		1. a reference to a preexisting, valid subobject: for it we use only restored selected properties.
+	 *		2. a reference to an old, deleted subobject which still exists in memory: for it we restore all properties.
+	 * 
+	 * See SnapshotUtil::Object::ResolveObjectDependencyForEditorWorld for more info.
+	 */
+	void MarkSubobjectForRestoringReferencesButSkipProperties(UObject* WorldSubbject);
+	bool IsSubobjectMarkedForReferenceRestorationOnly(const FSoftObjectPath& WorldSubobjectPath) const;
+
+
+	void AddComponentSelection(AActor* EditorWorldActor, const FAddedAndRemovedComponentInfo& ComponentSelection);
+	void RemoveComponentSelection(AActor* EditorWorldActor);
+
+	/*************** Custom subobject support ***************/
+
+	/** Recreate a Custom subobject that exists in snapshot world but is missing from editor world. Used for subobjects discovered via ICustomObjectSnapshotSerializer */
+	void AddCustomEditorSubobjectToRecreate(UObject* EditorWorldOwner, UObject* SnapshotSubobject);
+	void RemoveCustomEditorSubobjectToRecreate(UObject* EditorWorldOwner, UObject* SnapshotSubobject);
+
+	/*************** Queries ***************/
+
+	/** Will become DEPRECATED: Use GetObjectSelection(WorldObjectPath).GetObjectSelection()->GetPropertySelection() instead. */
 	const FPropertySelection* GetSelectedProperties(const FSoftObjectPath& WorldObjectPath) const;
+	FRestorableObjectSelection GetObjectSelection(const FSoftObjectPath& EditorWorldObject) const;
 	TArray<FSoftObjectPath> GetKeys() const;
+
+	bool HasChanges(AActor* EditorWorldActor) const;
+	bool HasChanges(UActorComponent* EditorWorldComponent) const;
+	bool HasChanges(UObject* EditorWorldObject) const;
 	
-	int32 GetKeyCount() const { return SelectedWorldObjectsToSelectedProperties.Num(); }
+	int32 GetKeyCount() const { return EditorWorldObjectToSelectedProperties.Num(); }
 	const TSet<FSoftObjectPath>& GetDeletedActorsToRespawn() const { return DeletedActorsToRespawn; }
 	const TSet<TWeakObjectPtr<AActor>>& GetNewActorsToDespawn() const { return NewActorsToDespawn; }
 
@@ -41,18 +82,33 @@ struct LEVELSNAPSHOTS_API FPropertySelectionMap
 	
 private:
 	
-	/* Maps a world actor to the properties that should be restored to values in a snapshot.
-	 * The properties are located in the actor itself or any other subcontainer (structs, components, or subobjects)
+	/**
+	 * Maps an editor world actor to the properties that should be restored to values in a snapshot.
+	 * The properties are located in the actor itself. Subobjects have a separate key.
+	 */
+	TMap<FSoftObjectPath, FPropertySelection> EditorWorldObjectToSelectedProperties;
+
+	/**
+	 * Maps an editor world actor to the components that should be added or removed.
+	 */
+	TMap<FSoftObjectPath, FAddedAndRemovedComponentInfo> EditorActorToComponentSelection;
+
+	/**
+	 * Maps an editor world object to the custom subobjects discovered by ICustomObjectSnapshotSerializer which need restoring.
+	 */
+	TMap<FSoftObjectPath, FCustomSubobjectRestorationInfo> EditorWorldObjectToCustomSubobjectSelection;
+	
+	/**
+	 * Contains objects which we do not serialize into: only replace references to them.
 	 */
-	TMap<FSoftObjectPath, FPropertySelection> SelectedWorldObjectsToSelectedProperties;
+	TSet<FSoftObjectPath> SubobjectsMarkedForReferencesRestorationOnly;
 
-	/* These actors were removed since the snapshot was taken. Re-create them.
+	/**
+	 * These actors were removed since the snapshot was taken. Re-create them.
 	 * This contains the original objects paths stored in the snapshot.
 	 */
-	UPROPERTY()
 	TSet<FSoftObjectPath> DeletedActorsToRespawn;
 
-	/* These actors were added since the snapshot was taken. Remove them. */
-	UPROPERTY()
+	/** These actors were added since the snapshot was taken. Remove them. */
 	TSet<TWeakObjectPtr<AActor>> NewActorsToDespawn;
 };
diff --git a/Engine/Plugins/VirtualProduction/LevelSnapshots/Source/LevelSnapshots/Public/Restorability/ISnapshotRestorabilityOverrider.h b/Engine/Plugins/VirtualProduction/LevelSnapshots/Source/LevelSnapshots/Public/Restorability/ISnapshotRestorabilityOverrider.h
--- a/Engine/Plugins/VirtualProduction/LevelSnapshots/Source/LevelSnapshots/Public/Restorability/ISnapshotRestorabilityOverrider.h
+++ b/Engine/Plugins/VirtualProduction/LevelSnapshots/Source/LevelSnapshots/Public/Restorability/ISnapshotRestorabilityOverrider.h
@@ -9,7 +9,7 @@ class UActorComponent;
 struct FPropertyContext;
 
 /**
- * Exposes callbacks for deciding whether an actor, component, or property is exposed to the snapshot system.
+ * Exposes callbacks for deciding whether an actor or component is exposed to the snapshot system.
  * 
  * Supported actors, components, or properties are:
  *	- Captured and saved when a snapshot is taken
@@ -17,7 +17,7 @@ struct FPropertyContext;
  *	- Show up in the results view
  *	- Restored when a snapshot is applied
  *
- * To see which actors, components, and properties are supported by default, see FSnapshotRestorability.
+ * To see which actors, components, and subobjects are supported by default, see FSnapshotRestorability.
  */
 class LEVELSNAPSHOTS_API ISnapshotRestorabilityOverrider
 {
@@ -27,7 +27,7 @@ public:
 	{
 		/* The object in question is included but only if nobody else returned Disallow.  */
 		Allow,
-        /* The object is included if some other override returns Allow. If nobody returns Allow, the object is not included. */
+        /* Let other overrides decide. If every override return DoNotCare, default snapshot behaviour is applied. */
         DoNotCare,
         /* The object in question is never suitable and is not included. Other overriders cannot override this. */
         Disallow
diff --git a/Engine/Plugins/VirtualProduction/LevelSnapshots/Source/LevelSnapshots/Public/Restorability/Serialization/ICustomObjectSnapshotSerializer.h b/Engine/Plugins/VirtualProduction/LevelSnapshots/Source/LevelSnapshots/Public/Restorability/Serialization/ICustomObjectSnapshotSerializer.h
--- a/Engine/Plugins/VirtualProduction/LevelSnapshots/Source/LevelSnapshots/Public/Restorability/Serialization/ICustomObjectSnapshotSerializer.h
+++ b/Engine/Plugins/VirtualProduction/LevelSnapshots/Source/LevelSnapshots/Public/Restorability/Serialization/ICustomObjectSnapshotSerializer.h
@@ -28,17 +28,23 @@ public:
 	 * Note that all uproperties will still be restored normally as with all other objects.
 	 *
 	 * @param EditorObject The object being snapshotted. Same as DataStorage->GetSerializedObject().
-	 * @param DataStorage Use it for tracking additional data
+	 * @param DataStorage Data to save for EditorObject: serialize extra data about EditorObject and add subobjects.  
 	 */
 	virtual void OnTakeSnapshot(UObject* EditorObject, ICustomSnapshotSerializationData& DataStorage) = 0;
+
+
 	
     /**
      * Called when creating objects for the temporary snapshot world. This is called for every subobject added using ISnapshotObjectSerializer::AddSubobjectDependency.
      * 
      * This function must either find the subobject in SnapshotObject or recreate it. If the object is recreated, you must fix up any property references yourself.
      * After this function is called, properties will be serialized into this function's return value. After this, OnPostSerializeSnapshotSubobject is called.
+     *
+     * @param SnapshotObject The outer of the subobject
+     * @param ObjectData The data saved for the subobject. It's the data associated with the index returned by AddSubobjectSnapshot.
+     * @param DataStorage The data saved for the outer
      * 
-     * This function may return null, in which case the subobject is ignored.
+     * @result The found or recreated subobject. If null, the subobject is skipped.
      */
 	virtual UObject* FindOrRecreateSubobjectInSnapshotWorld(UObject* SnapshotObject, const ISnapshotSubobjectMetaData& ObjectData, const ICustomSnapshotSerializationData& DataStorage) = 0;
 	
@@ -47,27 +53,43 @@ public:
 	 * 
 	 * This function must either find the subobject in EditorObject or recreate it. If the object is recreated, you must fix up any property references yourself.
 	 * After this function is called, properties will be serialized into this function's return value. After this, OnPostSerializeEditorSubobject is called.
+	 *
+	 * @param SnapshotObject The outer of the subobject
+	 * @param ObjectData The data saved for the subobject. It's the data associated with the index returned by AddSubobjectSnapshot.
+	 * @param DataStorage The data saved for the outer
 	 * 
-	 * This function may return null, in which case the subobject is ignored.
+     * @result The found or recreated subobject. If null, the subobject is skipped.
 	 */
 	virtual UObject* FindOrRecreateSubobjectInEditorWorld(UObject* EditorObject, const ISnapshotSubobjectMetaData& ObjectData, const ICustomSnapshotSerializationData& DataStorage) = 0;
 
 	/**
 	 * Similar to FindOrRecreateSubobjectInEditorWorld, only that the subobject is not recreated if not present. Called when diffing against the world.
+	 *
+	 * @param SnapshotObject The outer of the subobject
+	 * @param ObjectData The data saved for the subobject. It's the data associated with the index returned by AddSubobjectSnapshot.
+	 * @param DataStorage The data saved for the outer
+	 * 
+	 * @result The subobject in the editor world. If you return null, the subobject will be recreated.
 	 */
 	virtual UObject* FindSubobjectInEditorWorld(UObject* EditorObject, const ISnapshotSubobjectMetaData& ObjectData, const ICustomSnapshotSerializationData& DataStorage) = 0;
+
+
 	
 	/** Optional. Called after GetOrRecreateSubobjectInSnapshotWorld when all properties have been serialized into the subobject. You can do any post processing here. */
 	virtual void OnPostSerializeSnapshotSubobject(UObject* Subobject, const ISnapshotSubobjectMetaData& ObjectData, const ICustomSnapshotSerializationData& DataStorage) {}
 
 	/** Optional. Called after GetOrRecreateSubobjectInEditorWorld when all properties have been serialized into the subobject. You can do any post processing here. */
 	virtual void OnPostSerializeEditorSubobject(UObject* Subobject, const ISnapshotSubobjectMetaData& ObjectData, const ICustomSnapshotSerializationData& DataStorage) {}
+
+
 	
-	/** Optional. Called before properties are applied to the object. */
-	virtual void PreApplySnapshotProperties(UObject* EditorObject, const ICustomSnapshotSerializationData& DataStorage) {}
+	/** Optional. Called before properties are applied to an object for which this interface was registered. Called on both snapshot and editor world versions. */
+	virtual void PreApplySnapshotProperties(UObject* Object, const ICustomSnapshotSerializationData& DataStorage) {}
 	
-	/** Optional. Called after properties are applied to the object. */
-	virtual void PostApplySnapshotProperties(UObject* EditorObject, const ICustomSnapshotSerializationData& DataStorage) {}
+	/** Optional. Called after properties are applied to an object for which this interface was registered. Called on both snapshot and editor world versions. */
+	virtual void PostApplySnapshotProperties(UObject* Object, const ICustomSnapshotSerializationData& DataStorage) {}
+
+
 	
 	virtual ~ICustomObjectSnapshotSerializer() = default;
 };
\ No newline at end of file
diff --git a/Engine/Plugins/VirtualProduction/LevelSnapshots/Source/LevelSnapshots/Public/Restorability/SnapshotRestorability.h b/Engine/Plugins/VirtualProduction/LevelSnapshots/Source/LevelSnapshots/Public/Restorability/SnapshotRestorability.h
--- a/Engine/Plugins/VirtualProduction/LevelSnapshots/Source/LevelSnapshots/Public/Restorability/SnapshotRestorability.h
+++ b/Engine/Plugins/VirtualProduction/LevelSnapshots/Source/LevelSnapshots/Public/Restorability/SnapshotRestorability.h
@@ -17,9 +17,15 @@ public:
 
 	/* Is this actor captured by the snapshot system? */
 	static bool IsActorDesirableForCapture(const AActor* Actor);
-	/* Is this actor captured by the snapshot system? */
+	/* Is this component captured by the snapshot system? */
 	static bool IsComponentDesirableForCapture(const UActorComponent* Component);
-	
+	/* Is this subobject class captured by the snapshot system?*/
+	static bool IsSubobjectClassDesirableForCapture(const UClass* SubobjectClass);
+	/** Is this subobject captured by the snapshot system? */
+	static bool IsSubobjectDesirableForCapture(const UObject* Subobject);
+
+	/** Can the property be captured? */
+	static bool IsPropertyDesirableForCapture(const FProperty* Property);
 	/* Is this property never captured by the snapshot system? */
 	static bool IsPropertyBlacklistedForCapture(const FProperty* Property);
 	/* Is this property always captured by the snapshot system? */
diff --git a/Engine/Plugins/VirtualProduction/LevelSnapshots/Source/LevelSnapshots/Public/Settings/LevelSnapshotsEditorProjectSettings.h b/Engine/Plugins/VirtualProduction/LevelSnapshots/Source/LevelSnapshots/Public/Settings/LevelSnapshotsEditorProjectSettings.h
--- a/Engine/Plugins/VirtualProduction/LevelSnapshots/Source/LevelSnapshots/Public/Settings/LevelSnapshotsEditorProjectSettings.h
+++ b/Engine/Plugins/VirtualProduction/LevelSnapshots/Source/LevelSnapshots/Public/Settings/LevelSnapshotsEditorProjectSettings.h
@@ -17,9 +17,26 @@ public:
 	/* Setting the Window Size through code will not save the size to the config. To make sure it's saved, call SaveConfig(). */
 	void SetLastCreationWindowSize(const FVector2D InLastSize);
 
+	//~ Begin UObject Interface
+	virtual void PostInitProperties() override;
+#if WITH_EDITOR
+	virtual void PostEditChangeProperty(FPropertyChangedEvent& PropertyChangedEvent) override;
+#endif
+	//~ End UObject Interface
+	
+
 	/* Specifies classes and properties that should never be captured nor restored. */
 	UPROPERTY(config, EditAnywhere, Category = "Level Snapshots|Behavior")
 	FRestorationBlacklist Blacklist;
+
+	/** Used when comparing float properties. Floats that have changes beyond this point do not show up as changed. */
+	UPROPERTY(Config, EditAnywhere, Category = "Level Snapshots|Behavior", meta = (ClampMin = "0.00000001", ClampMax = "0.1")) // Max value is SMALL_NUMBER = 1e-8
+	float FloatComparisonPrecision = 1e-03f;
+
+	/** Used when comparing double properties. Doubles that have changes beyond this point do not show up as changed. */
+	UPROPERTY(Config, EditAnywhere, Category = "Level Snapshots|Behavior", meta = (ClampMin = "0.00000001", ClampMax = "0.1")) // Max value is SMALL_NUMBER = 1e-8
+	double DoubleComparisonPrecision = 1e-03;
+	
 	
 	UPROPERTY(config, EditAnywhere, Category = "Level Snapshots|Editor", meta = (ConfigRestartRequired = true))
 	bool bEnableLevelSnapshotsToolbarButton;
diff --git a/Engine/Plugins/VirtualProduction/LevelSnapshots/Source/LevelSnapshotsEditor/LevelSnapshotsEditor.build.cs b/Engine/Plugins/VirtualProduction/LevelSnapshots/Source/LevelSnapshotsEditor/LevelSnapshotsEditor.build.cs
--- a/Engine/Plugins/VirtualProduction/LevelSnapshots/Source/LevelSnapshotsEditor/LevelSnapshotsEditor.build.cs
+++ b/Engine/Plugins/VirtualProduction/LevelSnapshots/Source/LevelSnapshotsEditor/LevelSnapshotsEditor.build.cs
@@ -9,28 +9,6 @@ public class LevelSnapshotsEditor : ModuleRules
 	{
 		PCHUsage = ModuleRules.PCHUsageMode.UseExplicitOrSharedPCHs;
 		
-		PrivateIncludePaths.AddRange(
-			new string[] 
-			{
-				Path.Combine(ModuleDirectory, "Private"),
-				Path.Combine(ModuleDirectory, "Private", "AssetTypeActions"),
-				Path.Combine(ModuleDirectory, "Private", "Customizations"),
-				Path.Combine(ModuleDirectory, "Private", "Data"),
-				Path.Combine(ModuleDirectory, "Private", "Data", "Filters"),
-				Path.Combine(ModuleDirectory, "Private", "Data", "DragDrop"),
-				Path.Combine(ModuleDirectory, "Private", "Factories"),
-				Path.Combine(ModuleDirectory, "Private", "Misc"),
-				Path.Combine(ModuleDirectory, "Private", "TempInterfaces"),
-				Path.Combine(ModuleDirectory, "Private", "Toolkits"),
-				Path.Combine(ModuleDirectory, "Private", "Views"),
-				Path.Combine(ModuleDirectory, "Private", "Views", "Filter"),
-				Path.Combine(ModuleDirectory, "Private", "Views", "Input"),
-				Path.Combine(ModuleDirectory, "Private", "Views", "Results"),
-				Path.Combine(ModuleDirectory, "Private", "Widgets"),
-				Path.Combine(ModuleDirectory, "Private", "Widgets", "Filter")
-			}
-		);
-
 		PublicDependencyModuleNames.AddRange(
 			new string[]
 			{
diff --git a/Engine/Plugins/VirtualProduction/LevelSnapshots/Source/LevelSnapshotsEditor/Private/AssetTypeActions/AssetTypeActions_LevelSnapshot.cpp b/Engine/Plugins/VirtualProduction/LevelSnapshots/Source/LevelSnapshotsEditor/Private/AssetTypeActions/AssetTypeActions_LevelSnapshot.cpp
--- a/Engine/Plugins/VirtualProduction/LevelSnapshots/Source/LevelSnapshotsEditor/Private/AssetTypeActions/AssetTypeActions_LevelSnapshot.cpp
+++ b/Engine/Plugins/VirtualProduction/LevelSnapshots/Source/LevelSnapshotsEditor/Private/AssetTypeActions/AssetTypeActions_LevelSnapshot.cpp
@@ -3,8 +3,8 @@
 #include "AssetTypeActions_LevelSnapshot.h"
 
 #include "LevelSnapshot.h"
-#include "LevelSnapshotsEditor/Private/LevelSnapshotsEditorModule.h"
-#include "LevelSnapshotsEditorData.h"
+#include "LevelSnapshotsEditorModule.h"
+#include "Data/LevelSnapshotsEditorData.h"
 #include "LevelSnapshotsEditorFunctionLibrary.h"
 #include "LevelSnapshotsEditorStyle.h"
 
diff --git a/Engine/Plugins/VirtualProduction/LevelSnapshots/Source/LevelSnapshotsEditor/Private/Customizations/NegatableFilterDetailsCustomization.cpp b/Engine/Plugins/VirtualProduction/LevelSnapshots/Source/LevelSnapshotsEditor/Private/Customizations/NegatableFilterDetailsCustomization.cpp
--- a/Engine/Plugins/VirtualProduction/LevelSnapshots/Source/LevelSnapshotsEditor/Private/Customizations/NegatableFilterDetailsCustomization.cpp
+++ b/Engine/Plugins/VirtualProduction/LevelSnapshots/Source/LevelSnapshotsEditor/Private/Customizations/NegatableFilterDetailsCustomization.cpp
@@ -1,8 +1,8 @@
 // Copyright Epic Games, Inc. All Rights Reserved.
 
-#include "NegatableFilterDetailsCustomization.h"
+#include "Customizations/NegatableFilterDetailsCustomization.h"
 
-#include "NegatableFilter.h"
+#include "Data/Filters/NegatableFilter.h"
 
 #include "DetailLayoutBuilder.h"
 #include "DetailCategoryBuilder.h"
diff --git a/Engine/Plugins/VirtualProduction/LevelSnapshots/Source/LevelSnapshotsEditor/Private/Data/DragDrop/FavoriteFilterDragDrop.cpp b/Engine/Plugins/VirtualProduction/LevelSnapshots/Source/LevelSnapshotsEditor/Private/Data/DragDrop/FavoriteFilterDragDrop.cpp
--- a/Engine/Plugins/VirtualProduction/LevelSnapshots/Source/LevelSnapshotsEditor/Private/Data/DragDrop/FavoriteFilterDragDrop.cpp
+++ b/Engine/Plugins/VirtualProduction/LevelSnapshots/Source/LevelSnapshotsEditor/Private/Data/DragDrop/FavoriteFilterDragDrop.cpp
@@ -1,12 +1,11 @@
 // Copyright Epic Games, Inc. All Rights Reserved.
 
-#include "FavoriteFilterDragDrop.h"
+#include "Data/DragDrop/FavoriteFilterDragDrop.h"
 
-#include "EditorStyleSet.h"
-#include "LevelSnapshotsEditorData.h"
-#include "NegatableFilter.h"
-#include "SLevelSnapshotsEditorFilterRow.h"
+#include "Data/LevelSnapshotsEditorData.h"
+#include "Widgets/SLevelSnapshotsEditorFilterRow.h"
 
+#include "EditorStyleSet.h"
 #include "Framework/Application/SlateApplication.h"
 #include "Widgets/Images/SImage.h"
 #include "Widgets/Layout/SBorder.h"
diff --git a/Engine/Plugins/VirtualProduction/LevelSnapshots/Source/LevelSnapshotsEditor/Private/Data/FilterListData.cpp b/Engine/Plugins/VirtualProduction/LevelSnapshots/Source/LevelSnapshotsEditor/Private/Data/FilterListData.cpp
--- a/Engine/Plugins/VirtualProduction/LevelSnapshots/Source/LevelSnapshotsEditor/Private/Data/FilterListData.cpp
+++ b/Engine/Plugins/VirtualProduction/LevelSnapshots/Source/LevelSnapshotsEditor/Private/Data/FilterListData.cpp
@@ -24,7 +24,7 @@ void FFilterListData::UpdateFilteredList(UWorld* World, ULevelSnapshot* FromSnap
     RelatedSnapshot = FromSnapshot;
 
 	// We expect the number of filtered actors & components to stay roughly the same: retain existing memory
-	ModifiedActorsSelectedProperties_AllowedByFilter.Empty(false);
+	ModifiedEditorObjectsSelectedProperties_AllowedByFilter.Empty(false);
 	ModifiedWorldActors_AllowedByFilter.Empty(ModifiedWorldActors_AllowedByFilter.Num());
 	RemovedOriginalActorPaths_AllowedByFilter.Empty(RemovedOriginalActorPaths_AllowedByFilter.Num());
 	AddedWorldActors_AllowedByFilter.Empty(AddedWorldActors_AllowedByFilter.Num());
@@ -39,8 +39,8 @@ void FFilterListData::UpdateFilteredList(UWorld* World, ULevelSnapshot* FromSnap
 
 void FFilterListData::ApplyFilterToFindSelectedProperties(AActor* WorldActor, ULevelSnapshotFilter* FilterToApply)
 {
-	const FPropertySelection* AllowedSelectedProperties = ModifiedActorsSelectedProperties_AllowedByFilter.GetSelectedProperties(WorldActor);
-	const FPropertySelection* DisallowedSelectedProperties = ModifiedActorsSelectedProperties_AllowedByFilter.GetSelectedProperties(WorldActor);
+	const FPropertySelection* AllowedSelectedProperties = ModifiedEditorObjectsSelectedProperties_AllowedByFilter.GetSelectedProperties(WorldActor);
+	const FPropertySelection* DisallowedSelectedProperties = ModifiedEditorObjectsSelectedProperties_AllowedByFilter.GetSelectedProperties(WorldActor);
 	
 	if (!ensure(WorldActor && FilterToApply) || AllowedSelectedProperties || DisallowedSelectedProperties)
 	{
@@ -62,11 +62,11 @@ void FFilterListData::ApplyFilterToFindSelectedProperties(AActor* WorldActor, UL
 	
 	if (bIsAllowedByFilters)
 	{
-		ULevelSnapshotsFunctionLibrary::ApplyFilterToFindSelectedProperties(RelatedSnapshot, ModifiedActorsSelectedProperties_AllowedByFilter, WorldActor, DeserializedActor.Get(), FilterToApply);
+		ULevelSnapshotsFunctionLibrary::ApplyFilterToFindSelectedProperties(RelatedSnapshot, ModifiedEditorObjectsSelectedProperties_AllowedByFilter, WorldActor, DeserializedActor.Get(), FilterToApply);
 	}
 	else
 	{
-		ULevelSnapshotsFunctionLibrary::ApplyFilterToFindSelectedProperties(RelatedSnapshot, ModifiedActorsSelectedProperties_DisallowedByFilter, WorldActor, DeserializedActor.Get(), FilterToApply);
+		ULevelSnapshotsFunctionLibrary::ApplyFilterToFindSelectedProperties(RelatedSnapshot, ModifiedEditorObjectsSelectedProperties_DisallowedByFilter, WorldActor, DeserializedActor.Get(), FilterToApply);
 	}
 }
 
diff --git a/Engine/Plugins/VirtualProduction/LevelSnapshots/Source/LevelSnapshotsEditor/Private/Data/FilterListData.h b/Engine/Plugins/VirtualProduction/LevelSnapshots/Source/LevelSnapshotsEditor/Private/Data/FilterListData.h
--- a/Engine/Plugins/VirtualProduction/LevelSnapshots/Source/LevelSnapshotsEditor/Private/Data/FilterListData.h
+++ b/Engine/Plugins/VirtualProduction/LevelSnapshots/Source/LevelSnapshotsEditor/Private/Data/FilterListData.h
@@ -33,12 +33,12 @@ struct FFilterListData
 	 */
 	TWeakObjectPtr<AActor> GetSnapshotCounterpartFor(TWeakObjectPtr<AActor> WorldActor) const;
 
-	const FPropertySelectionMap& GetModifiedActorsSelectedProperties_AllowedByFilter() const { return ModifiedActorsSelectedProperties_AllowedByFilter; }
+	const FPropertySelectionMap& GetModifiedEditorObjectsSelectedProperties_AllowedByFilter() const { return ModifiedEditorObjectsSelectedProperties_AllowedByFilter; }
 	const TSet<TWeakObjectPtr<AActor>>& GetModifiedActors_AllowedByFilter() const { return ModifiedWorldActors_AllowedByFilter; }
 	const TSet<FSoftObjectPath>& GetRemovedOriginalActorPaths_AllowedByFilter() const { return RemovedOriginalActorPaths_AllowedByFilter; }
 	const TSet<TWeakObjectPtr<AActor>>& GetAddedWorldActors_AllowedByFilter() const { return AddedWorldActors_AllowedByFilter; }
 
-	const FPropertySelectionMap& GetModifiedActorsSelectedProperties_DisallowedByFilter() const { return ModifiedActorsSelectedProperties_DisallowedByFilter; }
+	const FPropertySelectionMap& GetModifiedEditorObjectsSelectedProperties_DisallowedByFilter() const { return ModifiedEditorObjectsSelectedProperties_DisallowedByFilter; }
 	const TSet<TWeakObjectPtr<AActor>>& GetModifiedActors_DisallowedByFilter() const { return ModifiedWorldActors_DisallowedByFilter; }
 	const TSet<FSoftObjectPath>& GetRemovedOriginalActorPaths_DisallowedByFilter() const { return RemovedOriginalActorPaths_DisallowedByFilter; }
 	const TSet<TWeakObjectPtr<AActor>>& GetAddedWorldActors_DisallowedByFilter() const { return AddedWorldActors_DisallowedByFilter; }
@@ -56,7 +56,7 @@ private:
 	
 	/* Selected properties for actors allowed by filters. */
 	UPROPERTY()
-	FPropertySelectionMap ModifiedActorsSelectedProperties_AllowedByFilter;
+	FPropertySelectionMap ModifiedEditorObjectsSelectedProperties_AllowedByFilter;
 
 	/* Actors to show in filter results panel when "ShowUnchanged = false". */
 	UPROPERTY()
@@ -72,7 +72,7 @@ private:
 	
 	/* Selected properties for actors disallowed by filters. */
 	UPROPERTY()
-	FPropertySelectionMap ModifiedActorsSelectedProperties_DisallowedByFilter;
+	FPropertySelectionMap ModifiedEditorObjectsSelectedProperties_DisallowedByFilter;
 
 	/* Actors to show in filter results panel when "ShowUnchanged = true". */
 	UPROPERTY()
diff --git a/Engine/Plugins/VirtualProduction/LevelSnapshots/Source/LevelSnapshotsEditor/Private/Data/FilterLoader.cpp b/Engine/Plugins/VirtualProduction/LevelSnapshots/Source/LevelSnapshotsEditor/Private/Data/FilterLoader.cpp
--- a/Engine/Plugins/VirtualProduction/LevelSnapshots/Source/LevelSnapshotsEditor/Private/Data/FilterLoader.cpp
+++ b/Engine/Plugins/VirtualProduction/LevelSnapshots/Source/LevelSnapshotsEditor/Private/Data/FilterLoader.cpp
@@ -1,8 +1,8 @@
 // Copyright Epic Games, Inc. All Rights Reserved.
 
-#include "FilterLoader.h"
+#include "Data/FilterLoader.h"
 
-#include "LevelSnapshotsFilterPreset.h"
+#include "Data/Filters/LevelSnapshotsFilterPreset.h"
 
 #include "AssetRegistry/AssetRegistryModule.h"
 #include "EditorDirectories.h"
diff --git a/Engine/Plugins/VirtualProduction/LevelSnapshots/Source/LevelSnapshotsEditor/Private/Data/FilterLoader.h b/Engine/Plugins/VirtualProduction/LevelSnapshots/Source/LevelSnapshotsEditor/Private/Data/FilterLoader.h
--- a/Engine/Plugins/VirtualProduction/LevelSnapshots/Source/LevelSnapshotsEditor/Private/Data/FilterLoader.h
+++ b/Engine/Plugins/VirtualProduction/LevelSnapshots/Source/LevelSnapshotsEditor/Private/Data/FilterLoader.h
@@ -4,8 +4,8 @@
 
 #include "CoreMinimal.h"
 
-#include "LevelSnapshotsFilterPreset.h"
 #include "AssetRegistry/AssetData.h"
+#include "Data/Filters/LevelSnapshotsFilterPreset.h"
 #include "UObject/Object.h"
 #include "FilterLoader.generated.h"
 
diff --git a/Engine/Plugins/VirtualProduction/LevelSnapshots/Source/LevelSnapshotsEditor/Private/Data/LevelSnapshotsEditorData.h b/Engine/Plugins/VirtualProduction/LevelSnapshots/Source/LevelSnapshotsEditor/Private/Data/LevelSnapshotsEditorData.h
--- a/Engine/Plugins/VirtualProduction/LevelSnapshots/Source/LevelSnapshotsEditor/Private/Data/LevelSnapshotsEditorData.h
+++ b/Engine/Plugins/VirtualProduction/LevelSnapshots/Source/LevelSnapshotsEditor/Private/Data/LevelSnapshotsEditorData.h
@@ -7,8 +7,8 @@
 #include "UObject/StrongObjectPtr.h"
 
 #include "LevelSnapshot.h"
-#include "LevelSnapshotsFilterPreset.h"
-#include "NegatableFilter.h"
+#include "Data/Filters/LevelSnapshotsFilterPreset.h"
+#include "Data/Filters/NegatableFilter.h"
 #include "LevelSnapshotsEditorData.generated.h"
 
 class UFavoriteFilterContainer;
diff --git a/Engine/Plugins/VirtualProduction/LevelSnapshots/Source/LevelSnapshotsEditor/Private/Toolkits/LevelSnapshotsEditorToolkit.cpp b/Engine/Plugins/VirtualProduction/LevelSnapshots/Source/LevelSnapshotsEditor/Private/Toolkits/LevelSnapshotsEditorToolkit.cpp
--- a/Engine/Plugins/VirtualProduction/LevelSnapshots/Source/LevelSnapshotsEditor/Private/Toolkits/LevelSnapshotsEditorToolkit.cpp
+++ b/Engine/Plugins/VirtualProduction/LevelSnapshots/Source/LevelSnapshotsEditor/Private/Toolkits/LevelSnapshotsEditorToolkit.cpp
@@ -7,7 +7,6 @@
 #include "Data/LevelSnapshotsEditorData.h"
 #include "LevelSnapshotsEditorModule.h"
 #include "LevelSnapshotsEditorStyle.h"
-#include "Misc/LevelSnapshotsEditorContext.h"
 #include "Views/Filter/LevelSnapshotsEditorFilters.h"
 #include "Views/Filter/SLevelSnapshotsEditorFilters.h"
 #include "Views/Results/LevelSnapshotsEditorResults.h"
@@ -53,16 +52,13 @@ void FLevelSnapshotsEditorToolkit::Initialize(ULevelSnapshotsEditorData* InEdito
 
 	EditorData.Get()->ClearActiveSnapshot();
 
-	EditorContext = MakeShared<FLevelSnapshotsEditorContext>();
-	ViewBuilder = MakeShared<FLevelSnapshotsEditorViewBuilder>();
-	ViewBuilder->EditorContextPtr = EditorContext;
-	ViewBuilder->EditorDataPtr = InEditorData;
+	ViewBuildData.EditorDataPtr = InEditorData;
 
 	
 	// Initialize views
-	EditorInput		= MakeShared<FLevelSnapshotsEditorInput>(ViewBuilder.ToSharedRef());
-	EditorFilters	= MakeShared<FLevelSnapshotsEditorFilters>(ViewBuilder.ToSharedRef());
-	EditorResults	= MakeShared<FLevelSnapshotsEditorResults>(ViewBuilder.ToSharedRef());
+	EditorInput		= MakeShared<FLevelSnapshotsEditorInput>(ViewBuildData);
+	EditorFilters	= MakeShared<FLevelSnapshotsEditorFilters>(ViewBuildData);
+	EditorResults	= MakeShared<FLevelSnapshotsEditorResults>(ViewBuildData);
 
 
 	// Create our content
diff --git a/Engine/Plugins/VirtualProduction/LevelSnapshots/Source/LevelSnapshotsEditor/Private/Toolkits/LevelSnapshotsEditorToolkit.h b/Engine/Plugins/VirtualProduction/LevelSnapshots/Source/LevelSnapshotsEditor/Private/Toolkits/LevelSnapshotsEditorToolkit.h
--- a/Engine/Plugins/VirtualProduction/LevelSnapshots/Source/LevelSnapshotsEditor/Private/Toolkits/LevelSnapshotsEditorToolkit.h
+++ b/Engine/Plugins/VirtualProduction/LevelSnapshots/Source/LevelSnapshotsEditor/Private/Toolkits/LevelSnapshotsEditorToolkit.h
@@ -3,6 +3,8 @@
 #pragma once
 
 #include "CoreMinimal.h"
+
+#include "Views/SnapshotEditorViewData.h"
 #include "Toolkits/AssetEditorToolkit.h"
 
 class FLevelSnapshotsEditorInput;
@@ -13,7 +15,7 @@ class SCheckBox;
 class SDockTab;
 class ULevelSnapshotsEditorData;
 
-struct FLevelSnapshotsEditorViewBuilder;
+struct FSnapshotEditorViewData;
 
 class ULevelSnapshot;
 
@@ -61,13 +63,11 @@ private:
 	
 	TWeakObjectPtr<ULevelSnapshotsEditorData> EditorData;
 	
-	TSharedPtr<FLevelSnapshotsEditorViewBuilder> ViewBuilder;
+	FSnapshotEditorViewData ViewBuildData;
 
 	TSharedPtr<FLevelSnapshotsEditorInput> EditorInput;
 	TSharedPtr<FLevelSnapshotsEditorFilters> EditorFilters;
 	TSharedPtr<FLevelSnapshotsEditorResults> EditorResults;
 
-	TSharedPtr<FLevelSnapshotsEditorContext> EditorContext;
-
 	TSharedPtr<SCheckBox> SettingsButtonPtr;
 };
diff --git a/Engine/Plugins/VirtualProduction/LevelSnapshots/Source/LevelSnapshotsEditor/Private/Views/Filter/LevelSnapshotsEditorFilters.cpp b/Engine/Plugins/VirtualProduction/LevelSnapshots/Source/LevelSnapshotsEditor/Private/Views/Filter/LevelSnapshotsEditorFilters.cpp
--- a/Engine/Plugins/VirtualProduction/LevelSnapshots/Source/LevelSnapshotsEditor/Private/Views/Filter/LevelSnapshotsEditorFilters.cpp
+++ b/Engine/Plugins/VirtualProduction/LevelSnapshots/Source/LevelSnapshotsEditor/Private/Views/Filter/LevelSnapshotsEditorFilters.cpp
@@ -1,16 +1,8 @@
 // Copyright Epic Games, Inc. All Rights Reserved.
 
 #include "Views/Filter/LevelSnapshotsEditorFilters.h"
-
-#include "ILevelSnapshotsEditorView.h"
-#include "LevelSnapshotFilters.h"
-#include "NegatableFilter.h"
 #include "Views/Filter/SLevelSnapshotsEditorFilters.h"
 
-FLevelSnapshotsEditorFilters::FLevelSnapshotsEditorFilters(const TSharedRef<FLevelSnapshotsEditorViewBuilder>& InBuilder)
-	: BuilderPtr(InBuilder)
-{}
-
 TSharedRef<SLevelSnapshotsEditorFilters> FLevelSnapshotsEditorFilters::GetOrCreateWidget()
 {
 	if (!EditorFiltersWidget.IsValid())
diff --git a/Engine/Plugins/VirtualProduction/LevelSnapshots/Source/LevelSnapshotsEditor/Private/Views/Filter/LevelSnapshotsEditorFilters.h b/Engine/Plugins/VirtualProduction/LevelSnapshots/Source/LevelSnapshotsEditor/Private/Views/Filter/LevelSnapshotsEditorFilters.h
--- a/Engine/Plugins/VirtualProduction/LevelSnapshots/Source/LevelSnapshotsEditor/Private/Views/Filter/LevelSnapshotsEditorFilters.h
+++ b/Engine/Plugins/VirtualProduction/LevelSnapshots/Source/LevelSnapshotsEditor/Private/Views/Filter/LevelSnapshotsEditorFilters.h
@@ -3,30 +3,24 @@
 #pragma once
 
 #include "CoreMinimal.h"
+#include "Views/SnapshotEditorViewData.h"
 
-#include "LevelSnapshotFilters.h"
-
-struct FLevelSnapshotsEditorViewBuilder;
 class SLevelSnapshotsEditorFilters;
 class UNegatableFilter;
 
 class FLevelSnapshotsEditorFilters : public TSharedFromThis<FLevelSnapshotsEditorFilters>
 {
 public:
-
-	DECLARE_MULTICAST_DELEGATE_OneParam(FOnSetActiveFilter, ULevelSnapshotFilter* /* bInEnabled */);
 	
-	FLevelSnapshotsEditorFilters(const TSharedRef<FLevelSnapshotsEditorViewBuilder>& InBuilder);
+	FLevelSnapshotsEditorFilters(const FSnapshotEditorViewData& BuildData)
+		: ViewBuildData(BuildData)
+	{}
 
 	TSharedRef<SLevelSnapshotsEditorFilters> GetOrCreateWidget();
-	TSharedPtr<FLevelSnapshotsEditorViewBuilder> GetBuilder() const { return BuilderPtr.Pin(); }
+	const FSnapshotEditorViewData& GetViewBuildData() const { return ViewBuildData; }
 
 private:
 	
-	FOnSetActiveFilter OnSetActiveFilter;
-	TWeakObjectPtr<ULevelSnapshotFilter> ActiveFilterPtr;
-	
 	TSharedPtr<SLevelSnapshotsEditorFilters> EditorFiltersWidget;
-
-	TWeakPtr<FLevelSnapshotsEditorViewBuilder> BuilderPtr;
+	FSnapshotEditorViewData ViewBuildData;
 };
diff --git a/Engine/Plugins/VirtualProduction/LevelSnapshots/Source/LevelSnapshotsEditor/Private/Views/Filter/SLevelSnapshotsEditorFilters.cpp b/Engine/Plugins/VirtualProduction/LevelSnapshots/Source/LevelSnapshotsEditor/Private/Views/Filter/SLevelSnapshotsEditorFilters.cpp
--- a/Engine/Plugins/VirtualProduction/LevelSnapshots/Source/LevelSnapshotsEditor/Private/Views/Filter/SLevelSnapshotsEditorFilters.cpp
+++ b/Engine/Plugins/VirtualProduction/LevelSnapshots/Source/LevelSnapshotsEditor/Private/Views/Filter/SLevelSnapshotsEditorFilters.cpp
@@ -6,9 +6,9 @@
 #include "Data/FilteredResults.h"
 #include "Data/Filters/LevelSnapshotsFilterPreset.h"
 #include "Data/LevelSnapshotsEditorData.h"
-#include "LevelSnapshotsEditorFilters.h"
+#include "Views/Filter/LevelSnapshotsEditorFilters.h"
 #include "LevelSnapshotsEditorStyle.h"
-#include "TempInterfaces/ILevelSnapshotsEditorView.h"
+#include "Views/SnapshotEditorViewData.h"
 #include "Widgets/Filter/SFavoriteFilterList.h"
 #include "Widgets/SLevelSnapshotsEditorFilterRow.h"
 #include "Widgets/Filter/SSaveAndLoadFilters.h"
@@ -23,8 +23,6 @@
 #include "Widgets/Layout/SScrollBox.h"
 #include "Widgets/Notifications/SNotificationList.h"
 #include "Widgets/SBoxPanel.h"
-#include "Widgets/Views/STableRow.h"
-#include "Widgets/Views/STreeView.h"
 
 #define LOCTEXT_NAMESPACE "LevelSnapshotsEditor"
 
@@ -119,7 +117,7 @@ void SLevelSnapshotsEditorFilters::Construct(const FArguments& InArgs, const TSh
 	
 	
 	FiltersModelPtr = InFilters;
-	EditorDataPtr = InFilters->GetBuilder()->EditorDataPtr;
+	EditorDataPtr = InFilters->GetViewBuildData().EditorDataPtr;
 
 	// Create a property view
 	FPropertyEditorModule& EditModule = FModuleManager::Get().GetModuleChecked<FPropertyEditorModule>("PropertyEditor");
@@ -209,7 +207,7 @@ void SLevelSnapshotsEditorFilters::Construct(const FArguments& InArgs, const TSh
             .AutoHeight()
             .HAlign(HAlign_Fill)
             [
-				SAssignNew(FavoriteList, SFavoriteFilterList, InFilters->GetBuilder()->EditorDataPtr->GetFavoriteFilters(), GetEditorData())
+				SAssignNew(FavoriteList, SFavoriteFilterList, InFilters->GetViewBuildData().EditorDataPtr->GetFavoriteFilters(), GetEditorData())
             ]
 
             // Rows 
diff --git a/Engine/Plugins/VirtualProduction/LevelSnapshots/Source/LevelSnapshotsEditor/Private/Views/Input/LevelSnapshotsEditorInput.cpp b/Engine/Plugins/VirtualProduction/LevelSnapshots/Source/LevelSnapshotsEditor/Private/Views/Input/LevelSnapshotsEditorInput.cpp
--- a/Engine/Plugins/VirtualProduction/LevelSnapshots/Source/LevelSnapshotsEditor/Private/Views/Input/LevelSnapshotsEditorInput.cpp
+++ b/Engine/Plugins/VirtualProduction/LevelSnapshots/Source/LevelSnapshotsEditor/Private/Views/Input/LevelSnapshotsEditorInput.cpp
@@ -6,16 +6,11 @@
 
 #include "CoreUObject/Public/AssetRegistry/AssetData.h"
 
-FLevelSnapshotsEditorInput::FLevelSnapshotsEditorInput(const TSharedRef<FLevelSnapshotsEditorViewBuilder>& InBuilder)
-	: BuilderPtr(InBuilder)
-{
-}
-
 TSharedRef<SWidget> FLevelSnapshotsEditorInput::GetOrCreateWidget()
 {
 	if (!EditorInputWidget.IsValid())
 	{
-		SAssignNew(EditorInputWidget, SLevelSnapshotsEditorInput, SharedThis(this), BuilderPtr.Pin().ToSharedRef());
+		SAssignNew(EditorInputWidget, SLevelSnapshotsEditorInput, SharedThis(this), ViewBuildData);
 	}
 
 	return EditorInputWidget.ToSharedRef();
diff --git a/Engine/Plugins/VirtualProduction/LevelSnapshots/Source/LevelSnapshotsEditor/Private/Views/Input/LevelSnapshotsEditorInput.h b/Engine/Plugins/VirtualProduction/LevelSnapshots/Source/LevelSnapshotsEditor/Private/Views/Input/LevelSnapshotsEditorInput.h
--- a/Engine/Plugins/VirtualProduction/LevelSnapshots/Source/LevelSnapshotsEditor/Private/Views/Input/LevelSnapshotsEditorInput.h
+++ b/Engine/Plugins/VirtualProduction/LevelSnapshots/Source/LevelSnapshotsEditor/Private/Views/Input/LevelSnapshotsEditorInput.h
@@ -2,26 +2,28 @@
 
 #pragma once
 
-#include "ILevelSnapshotsEditorInput.h"
+#include "Views/SnapshotEditorViewData.h"
+#include "Templates/SharedPointer.h"
 
 class SLevelSnapshotsEditorInput;
 
 struct FAssetData;
 
-class FLevelSnapshotsEditorInput : public ILevelSnapshotsEditorInput
+class FLevelSnapshotsEditorInput : public TSharedFromThis<FLevelSnapshotsEditorInput>
 {
 public:
-	FLevelSnapshotsEditorInput(const TSharedRef<FLevelSnapshotsEditorViewBuilder>& InBuilder);
+	
+	FLevelSnapshotsEditorInput(const FSnapshotEditorViewData& ViewBuildData)
+		: ViewBuildData(ViewBuildData)
+	{}
 
-	//~ Begin ILevelSnapshotsEditorView Interface
-	virtual TSharedRef<SWidget> GetOrCreateWidget() override;
-	virtual TSharedRef<FLevelSnapshotsEditorViewBuilder> GetBuilder() const override { return BuilderPtr.Pin().ToSharedRef(); }
-	//~ End ILevelSnapshotsEditorView Interface
+	TSharedRef<SWidget> GetOrCreateWidget();
+	const FSnapshotEditorViewData& GetBuilder() const { return ViewBuildData; }
 
 	void OpenLevelSnapshotsDialogWithAssetSelected(const FAssetData& InAssetData) const;
 
 private:
+	
 	TSharedPtr<SLevelSnapshotsEditorInput> EditorInputWidget;
-
-	TWeakPtr<FLevelSnapshotsEditorViewBuilder> BuilderPtr;
+	FSnapshotEditorViewData ViewBuildData;
 };
diff --git a/Engine/Plugins/VirtualProduction/LevelSnapshots/Source/LevelSnapshotsEditor/Private/Views/Input/SLevelSnapshotsEditorInput.cpp b/Engine/Plugins/VirtualProduction/LevelSnapshots/Source/LevelSnapshotsEditor/Private/Views/Input/SLevelSnapshotsEditorInput.cpp
--- a/Engine/Plugins/VirtualProduction/LevelSnapshots/Source/LevelSnapshotsEditor/Private/Views/Input/SLevelSnapshotsEditorInput.cpp
+++ b/Engine/Plugins/VirtualProduction/LevelSnapshots/Source/LevelSnapshotsEditor/Private/Views/Input/SLevelSnapshotsEditorInput.cpp
@@ -2,13 +2,12 @@
 
 #include "Views/Input/SLevelSnapshotsEditorInput.h"
 
-#include "Editor.h"
-#include "ILevelSnapshotsEditorView.h"
+#include "Data/LevelSnapshotsEditorData.h"
 #include "Views/Input/LevelSnapshotsEditorInput.h"
-#include "LevelSnapshotsEditorData.h"
+#include "Views/SnapshotEditorViewData.h"
 #include "Widgets/SLevelSnapshotsEditorBrowser.h"
 
-#include "LevelSnapshot.h"
+#include "Editor.h"
 #include "Engine/World.h"
 #include "Widgets/SBoxPanel.h"
 
@@ -21,19 +20,16 @@ SLevelSnapshotsEditorInput::~SLevelSnapshotsEditorInput()
 	FEditorDelegates::OnMapOpened.Remove(OnMapOpenedDelegateHandle);
 }
 
-void SLevelSnapshotsEditorInput::Construct(const FArguments& InArgs, const TSharedRef<FLevelSnapshotsEditorInput>& InEditorInput, const TSharedRef<FLevelSnapshotsEditorViewBuilder>& InBuilder)
+void SLevelSnapshotsEditorInput::Construct(const FArguments& InArgs, const TSharedRef<FLevelSnapshotsEditorInput>& InEditorInput, const FSnapshotEditorViewData& InViewBuildData)
 {
-	EditorInputPtr = InEditorInput;
-	BuilderPtr = InBuilder;
+	ViewBuildData = InViewBuildData;
+	check(ViewBuildData.EditorDataPtr.IsValid());
 
-	check(BuilderPtr.Pin()->EditorDataPtr.IsValid());
 
 	OnMapOpenedDelegateHandle = FEditorDelegates::OnMapOpened.AddLambda([this] (const FString& InFileName, const bool bAsTemplate)
 	{
-		check(BuilderPtr.IsValid());
-		check(BuilderPtr.Pin()->EditorDataPtr.IsValid());
-	
-		OverrideWorld(BuilderPtr.Pin()->EditorDataPtr->GetEditorWorld());
+		check(ViewBuildData.EditorDataPtr.IsValid());
+		OverrideWorld(ViewBuildData.EditorDataPtr->GetEditorWorld());
 	});
 
 	ChildSlot
@@ -42,8 +38,8 @@ void SLevelSnapshotsEditorInput::Construct(const FArguments& InArgs, const TShar
 
 		+ SVerticalBox::Slot()
 		[
-			SAssignNew(EditorBrowserWidgetPtr, SLevelSnapshotsEditorBrowser, InBuilder)
-			.OwningWorldPath(BuilderPtr.Pin()->EditorDataPtr->GetEditorWorld())
+			SAssignNew(EditorBrowserWidgetPtr, SLevelSnapshotsEditorBrowser, InViewBuildData)
+			.OwningWorldPath(ViewBuildData.EditorDataPtr->GetEditorWorld())
 		]
 	];
 }
@@ -56,7 +52,7 @@ void SLevelSnapshotsEditorInput::OpenLevelSnapshotsDialogWithAssetSelected(const
 void SLevelSnapshotsEditorInput::OverrideWorld(FSoftObjectPath InNewContextPath)
 {
 	// Replace the Browser widget with new world context if world and builder pointer valid
-	if (!ensure(InNewContextPath.IsValid()) || !ensure(BuilderPtr.IsValid()))
+	if (!ensure(InNewContextPath.IsValid()))
 	{
 		UE_LOG(LogLevelSnapshots, Error,
 			TEXT("SLevelSnapshotsEditorInput::OverrideWorld: Unable to rebuild Snapshot Browser; InNewContext or BuilderPtr are invalid."));
@@ -70,7 +66,7 @@ void SLevelSnapshotsEditorInput::OverrideWorld(FSoftObjectPath InNewContextPath)
 		
 		EditorInputOuterVerticalBox->AddSlot()
 		[
-			SAssignNew(EditorBrowserWidgetPtr, SLevelSnapshotsEditorBrowser, BuilderPtr.Pin().ToSharedRef())
+			SAssignNew(EditorBrowserWidgetPtr, SLevelSnapshotsEditorBrowser, ViewBuildData)
 			.OwningWorldPath(InNewContextPath)
 		];
 	}
diff --git a/Engine/Plugins/VirtualProduction/LevelSnapshots/Source/LevelSnapshotsEditor/Private/Views/Input/SLevelSnapshotsEditorInput.h b/Engine/Plugins/VirtualProduction/LevelSnapshots/Source/LevelSnapshotsEditor/Private/Views/Input/SLevelSnapshotsEditorInput.h
--- a/Engine/Plugins/VirtualProduction/LevelSnapshots/Source/LevelSnapshotsEditor/Private/Views/Input/SLevelSnapshotsEditorInput.h
+++ b/Engine/Plugins/VirtualProduction/LevelSnapshots/Source/LevelSnapshotsEditor/Private/Views/Input/SLevelSnapshotsEditorInput.h
@@ -2,6 +2,7 @@
 
 #pragma once
 
+#include "Views/SnapshotEditorViewData.h"
 #include "Widgets/DeclarativeSyntaxSupport.h"
 #include "Widgets/SCompoundWidget.h"
 
@@ -11,7 +12,7 @@ class SLevelSnapshotsEditorContextPicker;
 class SVerticalBox;
 class UWorld;
 
-struct FLevelSnapshotsEditorViewBuilder;
+struct FSnapshotEditorViewData;
 
 class SLevelSnapshotsEditorInput : public SCompoundWidget
 {
@@ -24,18 +25,16 @@ public:
 
 	SLATE_END_ARGS()
 
-	void Construct(const FArguments& InArgs, const TSharedRef<FLevelSnapshotsEditorInput>& InEditorInput, const TSharedRef<FLevelSnapshotsEditorViewBuilder>& InBuilder);
+	void Construct(const FArguments& InArgs, const TSharedRef<FLevelSnapshotsEditorInput>& InEditorInput, const FSnapshotEditorViewData& InViewBuildData);
 
 	void OpenLevelSnapshotsDialogWithAssetSelected(const FAssetData& InAssetData) const;
 
 private:
+	
 	void OverrideWorld(FSoftObjectPath InNewContextPath);
 
 	FDelegateHandle OnMapOpenedDelegateHandle;
-	
-	TWeakPtr<FLevelSnapshotsEditorInput> EditorInputPtr;
-
-	TWeakPtr<FLevelSnapshotsEditorViewBuilder> BuilderPtr;
+	FSnapshotEditorViewData ViewBuildData;
 
 	TSharedPtr<SVerticalBox> EditorInputOuterVerticalBox;
 	TSharedPtr<SLevelSnapshotsEditorContextPicker> EditorContextPickerPtr;
diff --git a/Engine/Plugins/VirtualProduction/LevelSnapshots/Source/LevelSnapshotsEditor/Private/Views/Results/LevelSnapshotsEditorResults.cpp b/Engine/Plugins/VirtualProduction/LevelSnapshots/Source/LevelSnapshotsEditor/Private/Views/Results/LevelSnapshotsEditorResults.cpp
--- a/Engine/Plugins/VirtualProduction/LevelSnapshots/Source/LevelSnapshotsEditor/Private/Views/Results/LevelSnapshotsEditorResults.cpp
+++ b/Engine/Plugins/VirtualProduction/LevelSnapshots/Source/LevelSnapshotsEditor/Private/Views/Results/LevelSnapshotsEditorResults.cpp
@@ -6,19 +6,13 @@
 
 #include "Views/Results/SLevelSnapshotsEditorResults.h"
 
-FLevelSnapshotsEditorResults::FLevelSnapshotsEditorResults(const TSharedRef<FLevelSnapshotsEditorViewBuilder>& InBuilder)
-	: BuilderPtr(InBuilder)
-{
-}
-
 TSharedRef<SWidget> FLevelSnapshotsEditorResults::GetOrCreateWidget()
 {
 	if (!EditorResultsWidget.IsValid())
 	{
-		TSharedPtr<FLevelSnapshotsEditorViewBuilder> PinnedBuilder = BuilderPtr.Pin();
-		if (ensure(PinnedBuilder && PinnedBuilder->EditorDataPtr.IsValid()))
+		if (ensure(ViewBuildData.EditorDataPtr.IsValid()))
 		{
-			SAssignNew(EditorResultsWidget, SLevelSnapshotsEditorResults, BuilderPtr.Pin()->EditorDataPtr.Get());
+			SAssignNew(EditorResultsWidget, SLevelSnapshotsEditorResults, ViewBuildData.EditorDataPtr.Get());
 		}
 		else
 		{
@@ -28,13 +22,3 @@ TSharedRef<SWidget> FLevelSnapshotsEditorResults::GetOrCreateWidget()
 
 	return EditorResultsWidget.ToSharedRef();
 }
-
-void FLevelSnapshotsEditorResults::BuildSelectionSetFromSelectedPropertiesInEachActorGroup() const
-{
-	EditorResultsWidget->BuildSelectionSetFromSelectedPropertiesInEachActorGroup();
-}
-
-void FLevelSnapshotsEditorResults::RefreshResults() const
-{
-	EditorResultsWidget->RefreshResults();
-}
diff --git a/Engine/Plugins/VirtualProduction/LevelSnapshots/Source/LevelSnapshotsEditor/Private/Views/Results/LevelSnapshotsEditorResults.h b/Engine/Plugins/VirtualProduction/LevelSnapshots/Source/LevelSnapshotsEditor/Private/Views/Results/LevelSnapshotsEditorResults.h
--- a/Engine/Plugins/VirtualProduction/LevelSnapshots/Source/LevelSnapshotsEditor/Private/Views/Results/LevelSnapshotsEditorResults.h
+++ b/Engine/Plugins/VirtualProduction/LevelSnapshots/Source/LevelSnapshotsEditor/Private/Views/Results/LevelSnapshotsEditorResults.h
@@ -2,25 +2,24 @@
 
 #pragma once
 
-#include "ILevelSnapshotsEditorResults.h"
+#include "Views/SnapshotEditorViewData.h"
+#include "Templates/SharedPointer.h"
 
 class SLevelSnapshotsEditorResults;
 
-class FLevelSnapshotsEditorResults : public ILevelSnapshotsEditorResults
+class FLevelSnapshotsEditorResults : public TSharedFromThis<FLevelSnapshotsEditorResults>
 {
 public:
-	FLevelSnapshotsEditorResults(const TSharedRef<FLevelSnapshotsEditorViewBuilder>& InBuilder);
+	
+	FLevelSnapshotsEditorResults(const FSnapshotEditorViewData& ViewBuildData)
+		: ViewBuildData(ViewBuildData)
+	{}
 
-	//~ Begin ILevelSnapshotsEditorView Interface
-	virtual TSharedRef<SWidget> GetOrCreateWidget() override;
-	virtual TSharedRef<FLevelSnapshotsEditorViewBuilder> GetBuilder() const override { return BuilderPtr.Pin().ToSharedRef(); }
-	//~ End ILevelSnapshotsEditorView Interface
-
-	void BuildSelectionSetFromSelectedPropertiesInEachActorGroup() const;
-	void RefreshResults() const;
+	TSharedRef<SWidget> GetOrCreateWidget();
+	const FSnapshotEditorViewData& GetBuilder() const { return ViewBuildData; }
 
 private:
+	
 	TSharedPtr<SLevelSnapshotsEditorResults> EditorResultsWidget;
-
-	TWeakPtr<FLevelSnapshotsEditorViewBuilder> BuilderPtr;
+	FSnapshotEditorViewData ViewBuildData;
 };
diff --git a/Engine/Plugins/VirtualProduction/LevelSnapshots/Source/LevelSnapshotsEditor/Private/Views/Results/SLevelSnapshotsEditorResults.cpp b/Engine/Plugins/VirtualProduction/LevelSnapshots/Source/LevelSnapshotsEditor/Private/Views/Results/SLevelSnapshotsEditorResults.cpp
--- a/Engine/Plugins/VirtualProduction/LevelSnapshots/Source/LevelSnapshotsEditor/Private/Views/Results/SLevelSnapshotsEditorResults.cpp
+++ b/Engine/Plugins/VirtualProduction/LevelSnapshots/Source/LevelSnapshotsEditor/Private/Views/Results/SLevelSnapshotsEditorResults.cpp
@@ -330,9 +330,11 @@ void FLevelSnapshotsEditorResultsRow::GenerateActorGroupChildren(FPropertySelect
 					ParentRowType == FLevelSnapshotsEditorResultsRow::StructGroup || 
 					ParentRowType == FLevelSnapshotsEditorResultsRow::StructInSetOrArray ||
 					ParentRowType == FLevelSnapshotsEditorResultsRow::StructInMap);
-			
 			const bool bIsPropertyFilteredOut = !bIsParentRowContainer && !PropertiesThatPassFilter.Contains(PropertyField);
-			if (bIsPropertyFilteredOut) { return; }
+			if (bIsPropertyFilteredOut)
+			{
+				return;
+			}
 			
 			bool bFoundCounterpart = false;
 			TSharedPtr<FPropertyHandleHierarchy> PinnedCounterpartHierarchy;
@@ -516,7 +518,7 @@ void FLevelSnapshotsEditorResultsRow::GenerateActorGroupChildren(FPropertySelect
 					}
 				}
 
-				if (NewProperty->DoesRowRepresentGroup() && !NewProperty->GetChildRows().Num() && !InputHandle->IsCustomized())
+				if (NewProperty->DoesRowRepresentGroup() && NewProperty->GetChildRows().Num() == 0 && !PropertiesThatPassFilter.Contains(PropertyField))
 				{
 					// No valid children, destroy group
 					NewProperty.Reset();
@@ -914,28 +916,33 @@ void FLevelSnapshotsEditorResultsRow::GenerateActorGroupChildren(FPropertySelect
 		}
 	}
 
-	// SubObjects
-
+	// Custom subbjects handled by external ICustomObjectSnapshotSerializer implementation
+	TMap<UObject*, UObject*> WorldToSnapshotCustomSubobjects;
+	
 	ULevelSnapshot* ActiveSnapshot = ResultsViewPtr.IsValid() && ResultsViewPtr.Pin()->GetEditorDataPtr() && ResultsViewPtr.Pin()->GetEditorDataPtr()->GetActiveSnapshot().IsSet() ? 
 		ResultsViewPtr.Pin()->GetEditorDataPtr()->GetActiveSnapshot().GetValue() : nullptr;
-	TMap<UObject*, UObject*> WorldToSnapshotSubobjects;
 
-	TFunction<void(UObject*, UObject*)> MatchingSubobjectLambda = [&WorldToSnapshotSubobjects, ActiveSnapshot, &MatchingSubobjectLambda](UObject* SnapshotSubobject, UObject* EditorWorldSubobject)
+	const TFunction<void(UObject* UnmatchedSnapshotSubobject)> HandleUnmatchedCustomSnapshotSubobject = [](UObject* UnmatchedCustomSnapshotSubobject)
+	{
+		// TODO: UnmatchedCustomSnapshotSubobject has no counterpart in the editor world. Display a row for it.
+		// If the row is checked when the snapshot is applied, call FPropertySelectionMap::AddCustomEditorSubobjectToRecreate
+	};
+
+	const TFunction<void(UObject* SnapshotSubobject, UObject* EditorWorldSubobject)> HandleMatchedCustomSubobjectPair = [&WorldToSnapshotCustomSubobjects, ActiveSnapshot, &HandleMatchedCustomSubobjectPair, &HandleUnmatchedCustomSnapshotSubobject](UObject* SnapshotSubobject, UObject* EditorWorldSubobject)
 	{
 		check(EditorWorldSubobject);
 		check(SnapshotSubobject);
 		
-		WorldToSnapshotSubobjects.Add(EditorWorldSubobject, SnapshotSubobject);
-			
-		ULevelSnapshotsFunctionLibrary::ForEachMatchingCustomSubobjectPair(ActiveSnapshot, SnapshotSubobject, EditorWorldSubobject, MatchingSubobjectLambda);
+		WorldToSnapshotCustomSubobjects.Add(EditorWorldSubobject, SnapshotSubobject);
+		ULevelSnapshotsFunctionLibrary::ForEachMatchingCustomSubobjectPair(ActiveSnapshot, SnapshotSubobject, EditorWorldSubobject, HandleMatchedCustomSubobjectPair, HandleUnmatchedCustomSnapshotSubobject);
 	};
 
 	if (ActiveSnapshot)
 	{
-		ULevelSnapshotsFunctionLibrary::ForEachMatchingCustomSubobjectPair(ActiveSnapshot, SnapshotActorLocal, WorldActorLocal, MatchingSubobjectLambda);
+		ULevelSnapshotsFunctionLibrary::ForEachMatchingCustomSubobjectPair(ActiveSnapshot, SnapshotActorLocal, WorldActorLocal, HandleMatchedCustomSubobjectPair, [](UObject*){});
 	}
 
-	for (const TPair<UObject*, UObject*> SubObjectPair : WorldToSnapshotSubobjects)
+	for (const TPair<UObject*, UObject*> SubObjectPair : WorldToSnapshotCustomSubobjects)
 	{
 		if (SubObjectPair.Key && SubObjectPair.Value)
 		{
@@ -2356,7 +2363,7 @@ void SLevelSnapshotsEditorResults::BuildSelectionSetFromSelectedPropertiesInEach
 		}
 	};
 
-	FPropertySelectionMap PropertySelectionMap = FilterListData.GetModifiedActorsSelectedProperties_AllowedByFilter();
+	FPropertySelectionMap PropertySelectionMap = FilterListData.GetModifiedEditorObjectsSelectedProperties_AllowedByFilter();
 
 	// Modified actors
 	for (const FLevelSnapshotsEditorResultsRowPtr& Group : TreeViewModifiedActorGroupObjects)
@@ -2639,49 +2646,27 @@ void SLevelSnapshotsEditorResults::GenerateTreeView(const bool bSnapshotHasChang
 
 bool SLevelSnapshotsEditorResults::GenerateTreeViewChildren_ModifiedActors(FLevelSnapshotsEditorResultsRowPtr ModifiedActorsHeader, ULevelSnapshotFilter* UserFilters)
 {
-
-	
 	check(ModifiedActorsHeader);
 	
 	const TSet<TWeakObjectPtr<AActor>>& ActorsToConsider = FilterListData.GetModifiedActors_AllowedByFilter();
-
 	TSet<FSoftObjectPath> EvaluatedObjects;
 
 	for (const TWeakObjectPtr<AActor>& WeakWorldActor : ActorsToConsider)
 	{
-		if (!ensure(WeakWorldActor.IsValid()))
+		AActor* WorldActor = WeakWorldActor.IsValid() ? WeakWorldActor.Get() : nullptr;
+		if (!ensure(WorldActor) || !FSnapshotRestorability::IsActorDesirableForCapture(WorldActor))
 		{
 			continue;
 		}
 
-		AActor* WorldActor = WeakWorldActor.Get();
-
-		if (!ensure(WorldActor))
-		{
-			continue;
-		}
-
-		if (!FSnapshotRestorability::IsActorDesirableForCapture(WorldActor))
-		{
-			continue;
-		}
-
-		const int32 KeyCountBeforeFilter = FilterListData.GetModifiedActorsSelectedProperties_AllowedByFilter().GetKeyCount();
-
 		// Get remaining properties after filter
 		if (UserFilters)
 		{
 			FilterListData.ApplyFilterToFindSelectedProperties(WorldActor, UserFilters);
 		}
 
-		const FPropertySelectionMap& ModifiedSelectedActors = FilterListData.GetModifiedActorsSelectedProperties_AllowedByFilter();
-		const int32 KeyCountAfterFilter = ModifiedSelectedActors.GetKeyCount();
-		const int32 KeyCountDifference = KeyCountAfterFilter - KeyCountBeforeFilter;
-
-		// If keys have been added, then this actor or its children have properties that pass the filter
-		const bool bDoesActorHavePropertiesAfterFilter = KeyCountDifference > 0;
-
-		if (!bDoesActorHavePropertiesAfterFilter)
+		const FPropertySelectionMap& ModifiedSelectedActors = FilterListData.GetModifiedEditorObjectsSelectedProperties_AllowedByFilter();
+		if (!ModifiedSelectedActors.HasChanges(WorldActor))
 		{
 			continue;
 		}
@@ -2890,7 +2875,7 @@ void SLevelSnapshotsEditorResults::OnGetRowChildren(FLevelSnapshotsEditorResults
 		{
 			if (Row->GetIsTreeViewItemExpanded())
 			{
-				FPropertySelectionMap PropertySelectionMap = FilterListData.GetModifiedActorsSelectedProperties_AllowedByFilter();
+				FPropertySelectionMap PropertySelectionMap = FilterListData.GetModifiedEditorObjectsSelectedProperties_AllowedByFilter();
 				Row->GenerateActorGroupChildren(PropertySelectionMap);
 
 				OutChildren = Row->GetChildRows();
diff --git a/Engine/Plugins/VirtualProduction/LevelSnapshots/Source/LevelSnapshotsEditor/Private/Views/Results/SLevelSnapshotsEditorResults.h b/Engine/Plugins/VirtualProduction/LevelSnapshots/Source/LevelSnapshotsEditor/Private/Views/Results/SLevelSnapshotsEditorResults.h
--- a/Engine/Plugins/VirtualProduction/LevelSnapshots/Source/LevelSnapshotsEditor/Private/Views/Results/SLevelSnapshotsEditorResults.h
+++ b/Engine/Plugins/VirtualProduction/LevelSnapshots/Source/LevelSnapshotsEditor/Private/Views/Results/SLevelSnapshotsEditorResults.h
@@ -2,12 +2,11 @@
 
 #pragma once
 
-#include "Widgets/Input/SCheckBox.h"
 #include "Widgets/SCompoundWidget.h"
 #include "Widgets/SBoxPanel.h"
 #include "Widgets/Views/STreeView.h"
 
-#include "FilterListData.h"
+#include "Data/FilterListData.h"
 #include "Data/LevelSnapshotsEditorData.h"
 #include "PropertySelection.h"
 
diff --git a/Engine/Plugins/VirtualProduction/LevelSnapshots/Source/LevelSnapshotsEditor/Private/Views/SnapshotEditorViewData.h b/Engine/Plugins/VirtualProduction/LevelSnapshots/Source/LevelSnapshotsEditor/Private/Views/SnapshotEditorViewData.h
new file mode 100644
--- /dev/null
+++ b/Engine/Plugins/VirtualProduction/LevelSnapshots/Source/LevelSnapshotsEditor/Private/Views/SnapshotEditorViewData.h
@@ -0,0 +1,13 @@
+// Copyright Epic Games, Inc. All Rights Reserved.
+
+#pragma once
+
+#include "CoreMinimal.h"
+#include "UObject/WeakObjectPtr.h"
+
+class ULevelSnapshotsEditorData;
+
+struct FSnapshotEditorViewData
+{
+	TWeakObjectPtr<ULevelSnapshotsEditorData> EditorDataPtr = nullptr;
+};
diff --git a/Engine/Plugins/VirtualProduction/LevelSnapshots/Source/LevelSnapshotsEditor/Private/Widgets/Filter/SCreateNewFilterWidget.cpp b/Engine/Plugins/VirtualProduction/LevelSnapshots/Source/LevelSnapshotsEditor/Private/Widgets/Filter/SCreateNewFilterWidget.cpp
--- a/Engine/Plugins/VirtualProduction/LevelSnapshots/Source/LevelSnapshotsEditor/Private/Widgets/Filter/SCreateNewFilterWidget.cpp
+++ b/Engine/Plugins/VirtualProduction/LevelSnapshots/Source/LevelSnapshotsEditor/Private/Widgets/Filter/SCreateNewFilterWidget.cpp
@@ -1,12 +1,12 @@
 // Copyright Epic Games, Inc. All Rights Reserved.
 
-#include "SCreateNewFilterWidget.h"
+#include "Widgets/Filter/SCreateNewFilterWidget.h"
 
-#include "ConjunctionFilter.h"
-#include "EditorStyleSet.h"
+#include "Data/Filters/ConjunctionFilter.h"
+#include "Data/FavoriteFilterContainer.h"
+#include "Widgets/Filter/SFilterSearchMenu.h"
 
-#include "FavoriteFilterContainer.h"
-#include "SFilterSearchMenu.h"
+#include "EditorStyleSet.h"
 #include "Widgets/Input/SComboButton.h"
 #include "Widgets/Input/SSearchBox.h"
 
diff --git a/Engine/Plugins/VirtualProduction/LevelSnapshots/Source/LevelSnapshotsEditor/Private/Widgets/Filter/SFavoriteFilter.cpp b/Engine/Plugins/VirtualProduction/LevelSnapshots/Source/LevelSnapshotsEditor/Private/Widgets/Filter/SFavoriteFilter.cpp
--- a/Engine/Plugins/VirtualProduction/LevelSnapshots/Source/LevelSnapshotsEditor/Private/Widgets/Filter/SFavoriteFilter.cpp
+++ b/Engine/Plugins/VirtualProduction/LevelSnapshots/Source/LevelSnapshotsEditor/Private/Widgets/Filter/SFavoriteFilter.cpp
@@ -1,8 +1,8 @@
 // Copyright Epic Games, Inc. All Rights Reserved.
 
-#include "SFavoriteFilter.h"
+#include "Widgets/Filter/SFavoriteFilter.h"
 
-#include "FavoriteFilterDragDrop.h"
+#include "Data/DragDrop/FavoriteFilterDragDrop.h"
 
 #include "EditorStyleSet.h"
 #include "Widgets/Layout/SBorder.h"
diff --git a/Engine/Plugins/VirtualProduction/LevelSnapshots/Source/LevelSnapshotsEditor/Private/Widgets/Filter/SFavoriteFilter.h b/Engine/Plugins/VirtualProduction/LevelSnapshots/Source/LevelSnapshotsEditor/Private/Widgets/Filter/SFavoriteFilter.h
--- a/Engine/Plugins/VirtualProduction/LevelSnapshots/Source/LevelSnapshotsEditor/Private/Widgets/Filter/SFavoriteFilter.h
+++ b/Engine/Plugins/VirtualProduction/LevelSnapshots/Source/LevelSnapshotsEditor/Private/Widgets/Filter/SFavoriteFilter.h
@@ -3,15 +3,16 @@
 #pragma once
 
 #include "CoreMinimal.h"
+#include "Data/LevelSnapshotsEditorData.h"
 
-#include "LevelSnapshotsEditorData.h"
 #include "Widgets/DeclarativeSyntaxSupport.h"
 #include "Widgets/SCompoundWidget.h"
 #include "Templates/SubclassOf.h"
 
 class FLevelSnapshotsEditorFilters;
 
-/* Display in the favorite filter list.
+/**
+ * Display in the favorite filter list.
  */
 class SFavoriteFilter : public SCompoundWidget
 {
diff --git a/Engine/Plugins/VirtualProduction/LevelSnapshots/Source/LevelSnapshotsEditor/Private/Widgets/Filter/SFavoriteFilterList.h b/Engine/Plugins/VirtualProduction/LevelSnapshots/Source/LevelSnapshotsEditor/Private/Widgets/Filter/SFavoriteFilterList.h
--- a/Engine/Plugins/VirtualProduction/LevelSnapshots/Source/LevelSnapshotsEditor/Private/Widgets/Filter/SFavoriteFilterList.h
+++ b/Engine/Plugins/VirtualProduction/LevelSnapshots/Source/LevelSnapshotsEditor/Private/Widgets/Filter/SFavoriteFilterList.h
@@ -5,7 +5,7 @@
 #include "Widgets/DeclarativeSyntaxSupport.h"
 #include "Widgets/SCompoundWidget.h"
 
-#include "FavoriteFilterContainer.h"
+#include "Data/FavoriteFilterContainer.h"
 
 class FLevelSnapshotsEditorFilters;
 class UFavoriteFilterContainer;
diff --git a/Engine/Plugins/VirtualProduction/LevelSnapshots/Source/LevelSnapshotsEditor/Private/Widgets/Filter/SFilterSearchMenu.cpp b/Engine/Plugins/VirtualProduction/LevelSnapshots/Source/LevelSnapshotsEditor/Private/Widgets/Filter/SFilterSearchMenu.cpp
--- a/Engine/Plugins/VirtualProduction/LevelSnapshots/Source/LevelSnapshotsEditor/Private/Widgets/Filter/SFilterSearchMenu.cpp
+++ b/Engine/Plugins/VirtualProduction/LevelSnapshots/Source/LevelSnapshotsEditor/Private/Widgets/Filter/SFilterSearchMenu.cpp
@@ -1,14 +1,14 @@
 // Copyright Epic Games, Inc. All Rights Reserved.
 
-#include "SFilterSearchMenu.h"
+#include "Widgets/Filter/SFilterSearchMenu.h"
 
-#include "FavoriteFilterContainer.h"
+#include "Data/FavoriteFilterContainer.h"
 #include "LevelSnapshotFilters.h"
+
 #include "Algo/Transform.h"
 #include "Framework/Application/SlateApplication.h"
 #include "Framework/Commands/UIAction.h"
 #include "Framework/MultiBox/MultiBoxBuilder.h"
-
 #include "Widgets/Input/SSearchBox.h"
 
 #define LOCTEXT_NAMESPACE "LevelSnapshotsEditor"
diff --git a/Engine/Plugins/VirtualProduction/LevelSnapshots/Source/LevelSnapshotsEditor/Private/Widgets/Filter/SSaveAndLoadFilters.cpp b/Engine/Plugins/VirtualProduction/LevelSnapshots/Source/LevelSnapshotsEditor/Private/Widgets/Filter/SSaveAndLoadFilters.cpp
--- a/Engine/Plugins/VirtualProduction/LevelSnapshots/Source/LevelSnapshotsEditor/Private/Widgets/Filter/SSaveAndLoadFilters.cpp
+++ b/Engine/Plugins/VirtualProduction/LevelSnapshots/Source/LevelSnapshotsEditor/Private/Widgets/Filter/SSaveAndLoadFilters.cpp
@@ -1,21 +1,19 @@
 ﻿// Copyright Epic Games, Inc. All Rights Reserved.
 
-#include "SSaveAndLoadFilters.h"
+#include "Widgets/Filter/SSaveAndLoadFilters.h"
 
-#include "LevelSnapshotsFilterPreset.h"
-#include "FilterLoader.h"
+#include "Data/Filters/LevelSnapshotsFilterPreset.h"
+#include "Data/FilterLoader.h"
+#include "Data/LevelSnapshotsEditorData.h"
 
 #include "ContentBrowserModule.h"
 #include "Engine/AssetManager.h"
+#include "Framework/MultiBox/MultiBoxBuilder.h"
 #include "IContentBrowserSingleton.h"
-#include "LevelSnapshotsEditorData.h"
 #include "Misc/MessageDialog.h"
-#include "SAssetSearchBox.h"
 #include "SAssetDropTarget.h"
-#include "Framework/MultiBox/MultiBoxBuilder.h"
 #include "Styling/SlateIconFinder.h"
 #include "Widgets/Images/SImage.h"
-#include "Widgets/Input/SComboBox.h"
 
 #define LOCTEXT_NAMESPACE "LevelSnapshotsEditor"
 
diff --git a/Engine/Plugins/VirtualProduction/LevelSnapshots/Source/LevelSnapshotsEditor/Private/Widgets/SLevelSnapshotsEditorBrowser.cpp b/Engine/Plugins/VirtualProduction/LevelSnapshots/Source/LevelSnapshotsEditor/Private/Widgets/SLevelSnapshotsEditorBrowser.cpp
--- a/Engine/Plugins/VirtualProduction/LevelSnapshots/Source/LevelSnapshotsEditor/Private/Widgets/SLevelSnapshotsEditorBrowser.cpp
+++ b/Engine/Plugins/VirtualProduction/LevelSnapshots/Source/LevelSnapshotsEditor/Private/Widgets/SLevelSnapshotsEditorBrowser.cpp
@@ -2,15 +2,15 @@
 
 #include "Widgets/SLevelSnapshotsEditorBrowser.h"
 
-#include "ILevelSnapshotsEditorView.h"
 #include "LevelSnapshot.h"
-#include "LevelSnapshotsEditorData.h"
+#include "Data/LevelSnapshotsEditorData.h"
+#include "Views/SnapshotEditorViewData.h"
 
-#include "IContentBrowserSingleton.h"
-#include "ContentBrowserModule.h"
-#include "LevelSnapshotsLog.h"
 #include "AssetRegistry/AssetRegistryModule.h"
+#include "ContentBrowserModule.h"
 #include "Editor/EditorStyle/Public/EditorStyleSet.h"
+#include "IContentBrowserSingleton.h"
+#include "LevelSnapshotsLog.h"
 #include "Misc/ScopedSlowTask.h"
 #include "Modules/ModuleManager.h"
 #include "Slate/Public/Framework/MultiBox/MultiBoxBuilder.h"
@@ -18,10 +18,10 @@
 
 #define LOCTEXT_NAMESPACE "LevelSnapshotsEditor"
 
-void SLevelSnapshotsEditorBrowser::Construct(const FArguments& InArgs, const TSharedRef<FLevelSnapshotsEditorViewBuilder>& InBuilder)
+void SLevelSnapshotsEditorBrowser::Construct(const FArguments& InArgs, const FSnapshotEditorViewData& InViewBuildData)
 {
 	OwningWorldPathAttribute = InArgs._OwningWorldPath;
-	BuilderPtr = InBuilder;
+	ViewBuildData = InViewBuildData;
 
 	check(OwningWorldPathAttribute.IsSet());
 
@@ -63,13 +63,12 @@ void SLevelSnapshotsEditorBrowser::SelectAsset(const FAssetData& InAssetData) co
 	SelectSnapshot.EnterProgressFrame(60.f);
 	SelectSnapshot.MakeDialog();
 	
-	TSharedPtr<FLevelSnapshotsEditorViewBuilder> Builder = BuilderPtr.Pin();
 	ULevelSnapshot* Snapshot = Cast<ULevelSnapshot>(InAssetData.GetAsset());
 
 	SelectSnapshot.EnterProgressFrame(40.f);
-	if (ensure(Builder) && ensure(Snapshot))
+	if (ensure(Snapshot))
 	{
-		Builder->EditorDataPtr->SetActiveSnapshot(Snapshot);
+		ViewBuildData.EditorDataPtr->SetActiveSnapshot(Snapshot);
 	}
 }
 
diff --git a/Engine/Plugins/VirtualProduction/LevelSnapshots/Source/LevelSnapshotsEditor/Private/Widgets/SLevelSnapshotsEditorBrowser.h b/Engine/Plugins/VirtualProduction/LevelSnapshots/Source/LevelSnapshotsEditor/Private/Widgets/SLevelSnapshotsEditorBrowser.h
--- a/Engine/Plugins/VirtualProduction/LevelSnapshots/Source/LevelSnapshotsEditor/Private/Widgets/SLevelSnapshotsEditorBrowser.h
+++ b/Engine/Plugins/VirtualProduction/LevelSnapshots/Source/LevelSnapshotsEditor/Private/Widgets/SLevelSnapshotsEditorBrowser.h
@@ -2,12 +2,13 @@
 
 #pragma once
 
+#include "Views/SnapshotEditorViewData.h"
 #include "Widgets/DeclarativeSyntaxSupport.h"
 #include "Widgets/SCompoundWidget.h"
 
 struct FAssetData;
 class UWorld;
-struct FLevelSnapshotsEditorViewBuilder;
+struct FSnapshotEditorViewData;
 
 class SLevelSnapshotsEditorBrowser : public SCompoundWidget
 {
@@ -20,7 +21,7 @@ public:
 
 	SLATE_END_ARGS()
 
-	void Construct(const FArguments& InArgs, const TSharedRef<FLevelSnapshotsEditorViewBuilder>& InBuilder);
+	void Construct(const FArguments& InArgs, const FSnapshotEditorViewData& InViewBuildData);
 	void SelectAsset(const FAssetData& InAssetData) const;
 
 private:
@@ -30,5 +31,5 @@ private:
 
 	TAttribute<FSoftObjectPath> OwningWorldPathAttribute;
 
-	TWeakPtr<FLevelSnapshotsEditorViewBuilder> BuilderPtr;
+	FSnapshotEditorViewData ViewBuildData;
 };
diff --git a/Engine/Plugins/VirtualProduction/LevelSnapshots/Source/LevelSnapshotsEditor/Private/Widgets/SLevelSnapshotsEditorCreationForm.cpp b/Engine/Plugins/VirtualProduction/LevelSnapshots/Source/LevelSnapshotsEditor/Private/Widgets/SLevelSnapshotsEditorCreationForm.cpp
--- a/Engine/Plugins/VirtualProduction/LevelSnapshots/Source/LevelSnapshotsEditor/Private/Widgets/SLevelSnapshotsEditorCreationForm.cpp
+++ b/Engine/Plugins/VirtualProduction/LevelSnapshots/Source/LevelSnapshotsEditor/Private/Widgets/SLevelSnapshotsEditorCreationForm.cpp
@@ -2,17 +2,14 @@
 
 #include "SLevelSnapshotsEditorCreationForm.h"
 
+#include "Data/LevelSnapshotsEditorData.h"
+#include "LevelSnapshotsEditorStyle.h"
 #include "Settings/LevelSnapshotsEditorProjectSettings.h"
 #include "Settings/LevelSnapshotsEditorDataManagementSettings.h"
-#include "LevelSnapshotsEditorData.h"
-#include "LevelSnapshotsEditorStyle.h"
 
-#include "Editor.h"
 #include "EditorStyleSet.h"
-#include "Engine/Engine.h"
 #include "Framework/Application/SlateApplication.h"
 #include "IDetailsView.h"
-#include "ISettingsSection.h"
 #include "Modules/ModuleManager.h"
 #include "PropertyEditorModule.h"
 #include "Widgets/Images/SImage.h"
@@ -25,7 +22,8 @@
 
 TSharedRef<SWindow> SLevelSnapshotsEditorCreationForm::MakeAndShowCreationWindow(
 	const FCloseCreationFormDelegate& CallOnClose, 
-	ULevelSnapshotsEditorProjectSettings* InProjectSettings, ULevelSnapshotsEditorDataManagementSettings* InDataManagementSettings)
+	ULevelSnapshotsEditorProjectSettings* InProjectSettings,
+	ULevelSnapshotsEditorDataManagementSettings* InDataManagementSettings)
 {
 	check(InProjectSettings);
 	check(InDataManagementSettings);
diff --git a/Engine/Plugins/VirtualProduction/LevelSnapshots/Source/LevelSnapshotsEditor/Private/Widgets/SLevelSnapshotsEditorFilter.cpp b/Engine/Plugins/VirtualProduction/LevelSnapshots/Source/LevelSnapshotsEditor/Private/Widgets/SLevelSnapshotsEditorFilter.cpp
--- a/Engine/Plugins/VirtualProduction/LevelSnapshots/Source/LevelSnapshotsEditor/Private/Widgets/SLevelSnapshotsEditorFilter.cpp
+++ b/Engine/Plugins/VirtualProduction/LevelSnapshots/Source/LevelSnapshotsEditor/Private/Widgets/SLevelSnapshotsEditorFilter.cpp
@@ -1,16 +1,14 @@
 // Copyright Epic Games, Inc. All Rights Reserved.
 
-#include "SLevelSnapshotsEditorFilter.h"
+#include "Widgets/SLevelSnapshotsEditorFilter.h"
 
+#include "Data/Filters/NegatableFilter.h"
+#include "Data/LevelSnapshotsEditorData.h"
 #include "LevelSnapshotsEditorStyle.h"
-#include "LevelSnapshotsEditorFilters.h"
-#include "SLevelSnapshotsFilterCheckBox.h"
+#include "Widgets/Filter/SLevelSnapshotsFilterCheckBox.h"
+#include "Widgets/Filter/SHoverableFilterActions.h"
 
 #include "EditorStyleSet.h"
-#include "ILevelSnapshotsEditorView.h"
-#include "LevelSnapshotsEditorData.h"
-#include "NegatableFilter.h"
-#include "SHoverableFilterActions.h"
 #include "Widgets/Layout/SBorder.h"
 #include "Widgets/Layout/SWrapBox.h"
 #include "Widgets/SOverlay.h"
diff --git a/Engine/Plugins/VirtualProduction/LevelSnapshots/Source/LevelSnapshotsEditor/Private/Widgets/SLevelSnapshotsEditorFilter.h b/Engine/Plugins/VirtualProduction/LevelSnapshots/Source/LevelSnapshotsEditor/Private/Widgets/SLevelSnapshotsEditorFilter.h
--- a/Engine/Plugins/VirtualProduction/LevelSnapshots/Source/LevelSnapshotsEditor/Private/Widgets/SLevelSnapshotsEditorFilter.h
+++ b/Engine/Plugins/VirtualProduction/LevelSnapshots/Source/LevelSnapshotsEditor/Private/Widgets/SLevelSnapshotsEditorFilter.h
@@ -4,8 +4,8 @@
 
 #include "CoreMinimal.h"
 
-#include "LevelSnapshotsEditorData.h"
-#include "NegatableFilter.h"
+#include "Data/LevelSnapshotsEditorData.h"
+#include "Data/Filters/NegatableFilter.h"
 
 #include "Widgets/DeclarativeSyntaxSupport.h"
 #include "Widgets/SCompoundWidget.h"
diff --git a/Engine/Plugins/VirtualProduction/LevelSnapshots/Source/LevelSnapshotsEditor/Private/Widgets/SLevelSnapshotsEditorFilterList.cpp b/Engine/Plugins/VirtualProduction/LevelSnapshots/Source/LevelSnapshotsEditor/Private/Widgets/SLevelSnapshotsEditorFilterList.cpp
--- a/Engine/Plugins/VirtualProduction/LevelSnapshots/Source/LevelSnapshotsEditor/Private/Widgets/SLevelSnapshotsEditorFilterList.cpp
+++ b/Engine/Plugins/VirtualProduction/LevelSnapshots/Source/LevelSnapshotsEditor/Private/Widgets/SLevelSnapshotsEditorFilterList.cpp
@@ -2,12 +2,12 @@
 
 #include "Widgets/SLevelSnapshotsEditorFilterList.h"
 
-#include "ConjunctionFilter.h"
+#include "Data/Filters/ConjunctionFilter.h"
 #include "LevelSnapshotsEditorStyle.h"
-#include "SCreateNewFilterWidget.h"
-#include "SLevelSnapshotsEditorFilter.h"
-#include "Components/HorizontalBox.h"
+#include "Widgets/Filter/SCreateNewFilterWidget.h"
+#include "Widgets/SLevelSnapshotsEditorFilter.h"
 
+#include "Components/HorizontalBox.h"
 #include "Widgets/Layout/SWrapBox.h"
 
 #define LOCTEXT_NAMESPACE "LevelSnapshotsEditor"
diff --git a/Engine/Plugins/VirtualProduction/LevelSnapshots/Source/LevelSnapshotsEditor/Private/Widgets/SLevelSnapshotsEditorFilterList.h b/Engine/Plugins/VirtualProduction/LevelSnapshots/Source/LevelSnapshotsEditor/Private/Widgets/SLevelSnapshotsEditorFilterList.h
--- a/Engine/Plugins/VirtualProduction/LevelSnapshots/Source/LevelSnapshotsEditor/Private/Widgets/SLevelSnapshotsEditorFilterList.h
+++ b/Engine/Plugins/VirtualProduction/LevelSnapshots/Source/LevelSnapshotsEditor/Private/Widgets/SLevelSnapshotsEditorFilterList.h
@@ -4,7 +4,7 @@
 
 #include "CoreMinimal.h"
 
-#include "LevelSnapshotsEditorData.h"
+#include "Data/LevelSnapshotsEditorData.h"
 #include "Widgets/DeclarativeSyntaxSupport.h"
 #include "Widgets/SCompoundWidget.h"
 
