Fix IsPendingKill usage per Horde

diff --git a/Engine/Plugins/VirtualProduction/LevelSnapshots/Source/LevelSnapshots/Private/Data/Util/SnapshotObjectUtil.cpp b/Engine/Plugins/VirtualProduction/LevelSnapshots/Source/LevelSnapshots/Private/Data/Util/SnapshotObjectUtil.cpp
--- a/Engine/Plugins/VirtualProduction/LevelSnapshots/Source/LevelSnapshots/Private/Data/Util/SnapshotObjectUtil.cpp
+++ b/Engine/Plugins/VirtualProduction/LevelSnapshots/Source/LevelSnapshots/Private/Data/Util/SnapshotObjectUtil.cpp
@@ -111,7 +111,7 @@ namespace
 
 		if (SubobjectData.bWasBlacklisted)
 		{
-			return !ResolvedObject || ResolvedObject->IsPendingKill() || ResolvedObject->IsUnreachable() ? nullptr : ResolvedObject;
+			return !IsValid(ResolvedObject) || ResolvedObject->IsUnreachable() ? nullptr : ResolvedObject;
 		}
 		
 		UClass* ExpectedClass = SubobjectData.Class.ResolveClass();
@@ -130,7 +130,7 @@ namespace
 		}
 
 		// Need to clean rename dead objects otherwise we'll get a name clash when instantiating new object
-		const bool bIsReferencingDeadObject = ResolvedObject && (ResolvedObject->IsPendingKill() || ResolvedObject->IsUnreachable());
+		const bool bIsReferencingDeadObject = !IsValid(ResolvedObject) || ResolvedObject->IsUnreachable();
 		const bool bIsReferencingDifferentClass = ResolvedObject && !ResolvedObject->GetClass()->IsChildOf(ExpectedClass);
 		if (bIsReferencingDeadObject || bIsReferencingDifferentClass)
 		{
@@ -312,7 +312,7 @@ UObject* SnapshotUtil::Object::ResolveObjectDependencyForEditorWorld(FWorldSnaps
 	if (ResolvedObject && (ResolvedObject->IsA<AActor>() || ResolvedObject->IsA<UActorComponent>()))
 	{
 		// Might resolve to a dead object which still exists in memory
-		return ResolvedObject->IsPendingKill() ? nullptr : ResolvedObject;
+		return IsValid(ResolvedObject) ? ResolvedObject : nullptr;
 	}
 	
 	if (FSubobjectSnapshotData* const SubobjectData = WorldData.Subobjects.Find(ObjectPathIndex))
@@ -346,4 +346,4 @@ int32 SnapshotUtil::Object::AddObjectDependency(FWorldSnapshotData& WorldData, U
 	}
 	
 	return Result;
-}
\ No newline at end of file
+}
diff --git a/Engine/Plugins/VirtualProduction/LevelSnapshots/Source/LevelSnapshots/Private/Tests/ComponentTests.cpp b/Engine/Plugins/VirtualProduction/LevelSnapshots/Source/LevelSnapshots/Private/Tests/ComponentTests.cpp
--- a/Engine/Plugins/VirtualProduction/LevelSnapshots/Source/LevelSnapshots/Private/Tests/ComponentTests.cpp
+++ b/Engine/Plugins/VirtualProduction/LevelSnapshots/Source/LevelSnapshots/Private/Tests/ComponentTests.cpp
@@ -300,7 +300,7 @@ bool FRestoreInstancedComponent::RunTest(const FString& Parameters)
 				return Component->GetFName() == "InstancedAddedToActor";
 			});
 			TestTrue(TEXT("Instanced component was removed from actor"), RemovedComponent == nullptr);
-			TestTrue(TEXT("Removed component is marked from destroy"), InstancedAddedToActor->IsPendingKill());
+			TestTrue(TEXT("Removed component is marked from destroy"), !IsValid(InstancedAddedToActor));
 
 			// Component was added back
 			UActorComponent* const* PossibleRecreatedComponent = Actor->GetInstanceComponents().FindByPredicate([&](UActorComponent* Component)
