Fix usage of editor for non-editor build.

diff --git a/Engine/Plugins/VirtualProduction/LevelSnapshots/Source/LevelSnapshots/Private/Archive/ApplySnapshotDataArchiveV2.cpp b/Engine/Plugins/VirtualProduction/LevelSnapshots/Source/LevelSnapshots/Private/Archive/ApplySnapshotDataArchiveV2.cpp
--- a/Engine/Plugins/VirtualProduction/LevelSnapshots/Source/LevelSnapshots/Private/Archive/ApplySnapshotDataArchiveV2.cpp
+++ b/Engine/Plugins/VirtualProduction/LevelSnapshots/Source/LevelSnapshots/Private/Archive/ApplySnapshotDataArchiveV2.cpp
@@ -210,7 +210,11 @@ void FApplySnapshotDataArchiveV2::PushSerializedProperty(FProperty* InProperty,
 
 UObject* FApplySnapshotDataArchiveV2::ResolveObjectDependency(int32 ObjectIndex) const
 {
+#if USE_STABLE_LOCALIZATION_KEYS
 	return SnapshotUtil::Object::ResolveObjectDependencyForEditorWorld(GetSharedData(), ObjectIndex, GetLocalizationNamespace(), SelectionMapForResolvingSubobjects);
+#else
+	return SnapshotUtil::Object::ResolveObjectDependencyForEditorWorld(GetSharedData(), ObjectIndex, "", SelectionMapForResolvingSubobjects);
+#endif
 }
 
 FApplySnapshotDataArchiveV2::FApplySnapshotDataArchiveV2(FObjectSnapshotData& InObjectData, FWorldSnapshotData& InSharedData, UObject* InOriginalObject, const FPropertySelectionMap& InSelectionMapForResolvingSubobjects, TOptional<const FPropertySelection*> InSelectionSet)
diff --git a/Engine/Plugins/VirtualProduction/LevelSnapshots/Source/LevelSnapshots/Private/Archive/SnapshotArchive.cpp b/Engine/Plugins/VirtualProduction/LevelSnapshots/Source/LevelSnapshots/Private/Archive/SnapshotArchive.cpp
--- a/Engine/Plugins/VirtualProduction/LevelSnapshots/Source/LevelSnapshots/Private/Archive/SnapshotArchive.cpp
+++ b/Engine/Plugins/VirtualProduction/LevelSnapshots/Source/LevelSnapshots/Private/Archive/SnapshotArchive.cpp
@@ -151,7 +151,11 @@ void FSnapshotArchive::Serialize(void* Data, int64 Length)
 
 UObject* FSnapshotArchive::ResolveObjectDependency(int32 ObjectIndex) const
 {
+#if USE_STABLE_LOCALIZATION_KEYS
 	return SnapshotUtil::Object::ResolveObjectDependencyForSnapshotWorld(SharedData, ObjectIndex, GetLocalizationNamespace());
+#else
+	return SnapshotUtil::Object::ResolveObjectDependencyForSnapshotWorld(SharedData, ObjectIndex, "");
+#endif
 }
 
 FSnapshotArchive::FSnapshotArchive(FObjectSnapshotData& InObjectData, FWorldSnapshotData& InSharedData, bool bIsLoading, UObject* InSerializedObject)
diff --git a/Engine/Plugins/VirtualProduction/LevelSnapshots/Source/LevelSnapshots/Private/Data/ActorSnapshotData.cpp b/Engine/Plugins/VirtualProduction/LevelSnapshots/Source/LevelSnapshots/Private/Data/ActorSnapshotData.cpp
--- a/Engine/Plugins/VirtualProduction/LevelSnapshots/Source/LevelSnapshots/Private/Data/ActorSnapshotData.cpp
+++ b/Engine/Plugins/VirtualProduction/LevelSnapshots/Source/LevelSnapshots/Private/Data/ActorSnapshotData.cpp
@@ -19,9 +19,11 @@
 #if USE_STABLE_LOCALIZATION_KEYS
 #include "Internationalization/TextPackageNamespaceUtil.h"
 #endif
+#if WITH_EDITOR
 #include "LevelEditor.h"
 #include "UnrealEdGlobals.h"
 #include "Editor/UnrealEdEngine.h"
+#endif
 #include "RestorationEvents/ApplySnapshotToActorScope.h"
 #include "RestorationEvents/RecreateComponentScope.h"
 #include "RestorationEvents/RemoveComponentScope.h"
