Fix non-editor build errors

diff --git a/Engine/Plugins/VirtualProduction/LevelSnapshots/Source/LevelSnapshots/Private/LevelSnapshotsModule.cpp b/Engine/Plugins/VirtualProduction/LevelSnapshots/Source/LevelSnapshots/Private/LevelSnapshotsModule.cpp
--- a/Engine/Plugins/VirtualProduction/LevelSnapshots/Source/LevelSnapshots/Private/LevelSnapshotsModule.cpp
+++ b/Engine/Plugins/VirtualProduction/LevelSnapshots/Source/LevelSnapshots/Private/LevelSnapshotsModule.cpp
@@ -48,12 +48,14 @@ namespace
 
 	void DisableIrrelevantBrushSubobjects(FLevelSnapshotsModule& Module)
 	{
+#if WITH_EDITORONLY_DATA
 		// ABrush::BrushBuilder is CPF_Edit but no user ever cares about it. We don't want it to make volumes to show up as changed.
 		const FProperty* BrushBuilder = ABrush::StaticClass()->FindPropertyByName(GET_MEMBER_NAME_CHECKED(ABrush, BrushBuilder));
 		if (ensure(BrushBuilder))
 		{
 			Module.AddBlacklistedProperties({ BrushBuilder });
 		}
+#endif
 	}
 
 	void DisableIrrelevantWorldSettings(FLevelSnapshotsModule& Module)
diff --git a/Engine/Plugins/VirtualProduction/LevelSnapshots/Source/LevelSnapshots/Private/Settings/LevelSnapshotsEditorProjectSettings.cpp b/Engine/Plugins/VirtualProduction/LevelSnapshots/Source/LevelSnapshots/Private/Settings/LevelSnapshotsEditorProjectSettings.cpp
--- a/Engine/Plugins/VirtualProduction/LevelSnapshots/Source/LevelSnapshots/Private/Settings/LevelSnapshotsEditorProjectSettings.cpp
+++ b/Engine/Plugins/VirtualProduction/LevelSnapshots/Source/LevelSnapshots/Private/Settings/LevelSnapshotsEditorProjectSettings.cpp
@@ -42,6 +42,7 @@ void ULevelSnapshotsEditorProjectSettings::PostInitProperties()
 	FPropertyInfoHelpers::UpdateDecimalComparisionPrecision(FloatComparisonPrecision, DoubleComparisonPrecision);
 }
 
+#if WITH_EDITOR
 void ULevelSnapshotsEditorProjectSettings::PostEditChangeProperty(FPropertyChangedEvent& PropertyChangedEvent)
 {
 	if (PropertyChangedEvent.GetPropertyName() == GET_MEMBER_NAME_CHECKED(ULevelSnapshotsEditorProjectSettings, FloatComparisonPrecision)
@@ -49,6 +50,7 @@ void ULevelSnapshotsEditorProjectSettings::PostEditChangeProperty(FPropertyChang
 	{
 		FPropertyInfoHelpers::UpdateDecimalComparisionPrecision(FloatComparisonPrecision, DoubleComparisonPrecision);
 	}
-	
+
 	UObject::PostEditChangeProperty(PropertyChangedEvent);
 }
+#endif
