Remove unused code.

diff --git a/Engine/Plugins/VirtualProduction/LevelSnapshots/Source/LevelSnapshots/Private/PropertySelectionMap.cpp b/Engine/Plugins/VirtualProduction/LevelSnapshots/Source/LevelSnapshots/Private/PropertySelectionMap.cpp
--- a/Engine/Plugins/VirtualProduction/LevelSnapshots/Source/LevelSnapshots/Private/PropertySelectionMap.cpp
+++ b/Engine/Plugins/VirtualProduction/LevelSnapshots/Source/LevelSnapshots/Private/PropertySelectionMap.cpp
@@ -113,12 +113,11 @@ TArray<FSoftObjectPath> FPropertySelectionMap::GetKeys() const
 
 	TArray<FSoftObjectPath> ComponentPaths;
 	EditorActorToComponentSelection.GenerateKeyArray(ComponentPaths);
-	LIKELY(ComponentPaths.Num() < Result.Num());
 	for (const FSoftObjectPath& Path : ComponentPaths)
 	{
 		Result.Add(Path);
 	}
-	
+
 	return Result;
 }
 
