Fix non-unity errors

diff --git a/Engine/Plugins/VirtualProduction/LevelSnapshots/Source/LevelSnapshots/Private/Data/ActorSnapshotData.cpp b/Engine/Plugins/VirtualProduction/LevelSnapshots/Source/LevelSnapshots/Private/Data/ActorSnapshotData.cpp
--- a/Engine/Plugins/VirtualProduction/LevelSnapshots/Source/LevelSnapshots/Private/Data/ActorSnapshotData.cpp
+++ b/Engine/Plugins/VirtualProduction/LevelSnapshots/Source/LevelSnapshots/Private/Data/ActorSnapshotData.cpp
@@ -1,29 +1,22 @@
 // Copyright Epic Games, Inc. All Rights Reserved.
 
-#include "ActorSnapshotData.h"
+#include "Data/ActorSnapshotData.h"
 
-#include "ApplySnapshotDataArchiveV2.h"
+#include "Archive/ApplySnapshotDataArchiveV2.h"
+#include "Archive/TakeWorldObjectSnapshotArchive.h"
+#include "Data/WorldSnapshotData.h"
+#include "Data/SnapshotCustomVersion.h"
 #include "CustomSerialization/CustomObjectSerializationWrapper.h"
-#include "CustomSerialization/CustomSerializationDataManager.h"
 #include "LevelSnapshotsLog.h"
+#include "LevelSnapshotsModule.h"
 #include "PropertySelectionMap.h"
-#include "SnapshotCustomVersion.h"
-#include "SnapshotRestorability.h"
-#include "TakeWorldObjectSnapshotArchive.h"
-#include "WorldSnapshotData.h"
+#include "Restorability/SnapshotRestorability.h"
 #include "Util/SnapshotObjectUtil.h"
 #include "Util/SnapshotUtil.h"
 
 #include "Engine/World.h"
 #include "GameFramework/Actor.h"
-#if USE_STABLE_LOCALIZATION_KEYS
-#include "Internationalization/TextPackageNamespaceUtil.h"
-#endif
-#if WITH_EDITOR
-#include "LevelEditor.h"
-#include "UnrealEdGlobals.h"
-#include "Editor/UnrealEdEngine.h"
-#endif
+#include "Modules/ModuleManager.h"
 #include "RestorationEvents/ApplySnapshotToActorScope.h"
 #include "RestorationEvents/RecreateComponentScope.h"
 #include "RestorationEvents/RemoveComponentScope.h"
@@ -31,6 +24,15 @@
 #include "UObject/Package.h"
 #include "UObject/Script.h"
 
+#if USE_STABLE_LOCALIZATION_KEYS
+#include "Internationalization/TextPackageNamespaceUtil.h"
+#endif
+
+#if WITH_EDITOR
+#include "Editor/UnrealEdEngine.h"
+#include "LevelEditor.h"
+#endif
+
 namespace
 {	
 	void FindOrRecreateSavedComponents(const FActorSnapshotData& SnapshotData, FWorldSnapshotData& WorldData, AActor* Actor)
diff --git a/Engine/Plugins/VirtualProduction/LevelSnapshots/Source/LevelSnapshots/Private/Data/RestorationEvents/ApplySnapshotPropertiesScope.cpp b/Engine/Plugins/VirtualProduction/LevelSnapshots/Source/LevelSnapshots/Private/Data/RestorationEvents/ApplySnapshotPropertiesScope.cpp
--- a/Engine/Plugins/VirtualProduction/LevelSnapshots/Source/LevelSnapshots/Private/Data/RestorationEvents/ApplySnapshotPropertiesScope.cpp
+++ b/Engine/Plugins/VirtualProduction/LevelSnapshots/Source/LevelSnapshots/Private/Data/RestorationEvents/ApplySnapshotPropertiesScope.cpp
@@ -4,11 +4,11 @@
 
 #include "LevelSnapshotsModule.h"
 
-FApplySnapshotPropertiesScope::FApplySnapshotPropertiesScope(FApplySnapshotPropertiesParams Params)
+FApplySnapshotPropertiesScope::FApplySnapshotPropertiesScope(const FApplySnapshotPropertiesParams& InParams)
 	:
-	Params(Params)
+	Params(InParams)
 {
-	FLevelSnapshotsModule::GetInternalModuleInstance().OnPreApplySnapshotProperties(Params);
+	FLevelSnapshotsModule::GetInternalModuleInstance().OnPreApplySnapshotProperties(InParams);
 }
 
 FApplySnapshotPropertiesScope::~FApplySnapshotPropertiesScope()
diff --git a/Engine/Plugins/VirtualProduction/LevelSnapshots/Source/LevelSnapshots/Private/Data/RestorationEvents/ApplySnapshotPropertiesScope.h b/Engine/Plugins/VirtualProduction/LevelSnapshots/Source/LevelSnapshots/Private/Data/RestorationEvents/ApplySnapshotPropertiesScope.h
--- a/Engine/Plugins/VirtualProduction/LevelSnapshots/Source/LevelSnapshots/Private/Data/RestorationEvents/ApplySnapshotPropertiesScope.h
+++ b/Engine/Plugins/VirtualProduction/LevelSnapshots/Source/LevelSnapshots/Private/Data/RestorationEvents/ApplySnapshotPropertiesScope.h
@@ -16,6 +16,6 @@ class FApplySnapshotPropertiesScope : public FNoncopyable
 	const FApplySnapshotPropertiesParams Params;
 public:
 
-	FApplySnapshotPropertiesScope(FApplySnapshotPropertiesParams Params);
+	FApplySnapshotPropertiesScope(const FApplySnapshotPropertiesParams& InParams);
 	~FApplySnapshotPropertiesScope();
 };
\ No newline at end of file
diff --git a/Engine/Plugins/VirtualProduction/LevelSnapshots/Source/LevelSnapshots/Private/Data/RestorationEvents/ApplySnapshotToActorScope.cpp b/Engine/Plugins/VirtualProduction/LevelSnapshots/Source/LevelSnapshots/Private/Data/RestorationEvents/ApplySnapshotToActorScope.cpp
--- a/Engine/Plugins/VirtualProduction/LevelSnapshots/Source/LevelSnapshots/Private/Data/RestorationEvents/ApplySnapshotToActorScope.cpp
+++ b/Engine/Plugins/VirtualProduction/LevelSnapshots/Source/LevelSnapshots/Private/Data/RestorationEvents/ApplySnapshotToActorScope.cpp
@@ -4,7 +4,7 @@
 
 #include "LevelSnapshotsModule.h"
 
-FApplySnapshotToActorScope::FApplySnapshotToActorScope(FApplySnapshotToActorParams Params)
+FApplySnapshotToActorScope::FApplySnapshotToActorScope(const FApplySnapshotToActorParams& Params)
 	:
 	Params(Params)
 {
diff --git a/Engine/Plugins/VirtualProduction/LevelSnapshots/Source/LevelSnapshots/Private/Data/RestorationEvents/ApplySnapshotToActorScope.h b/Engine/Plugins/VirtualProduction/LevelSnapshots/Source/LevelSnapshots/Private/Data/RestorationEvents/ApplySnapshotToActorScope.h
--- a/Engine/Plugins/VirtualProduction/LevelSnapshots/Source/LevelSnapshots/Private/Data/RestorationEvents/ApplySnapshotToActorScope.h
+++ b/Engine/Plugins/VirtualProduction/LevelSnapshots/Source/LevelSnapshots/Private/Data/RestorationEvents/ApplySnapshotToActorScope.h
@@ -16,6 +16,6 @@ class FApplySnapshotToActorScope : public FNoncopyable
 	const FApplySnapshotToActorParams Params;
 public:
 
-	FApplySnapshotToActorScope(FApplySnapshotToActorParams Params);
+	FApplySnapshotToActorScope(const FApplySnapshotToActorParams& Params);
 	~FApplySnapshotToActorScope();
 };
\ No newline at end of file
diff --git a/Engine/Plugins/VirtualProduction/LevelSnapshots/Source/LevelSnapshots/Private/Data/RestorationEvents/RecreateComponentScope.cpp b/Engine/Plugins/VirtualProduction/LevelSnapshots/Source/LevelSnapshots/Private/Data/RestorationEvents/RecreateComponentScope.cpp
--- a/Engine/Plugins/VirtualProduction/LevelSnapshots/Source/LevelSnapshots/Private/Data/RestorationEvents/RecreateComponentScope.cpp
+++ b/Engine/Plugins/VirtualProduction/LevelSnapshots/Source/LevelSnapshots/Private/Data/RestorationEvents/RecreateComponentScope.cpp
@@ -4,7 +4,9 @@
 
 #include "LevelSnapshotsLog.h"
 #include "LevelSnapshotsModule.h"
-#include "SubobjectSnapshotData.h"
+#include "Data/SubobjectSnapshotData.h"
+
+#include "Components/ActorComponent.h"
 
 FRecreateComponentScope::FRecreateComponentScope(
 	const FSubobjectSnapshotData& SubobjectSnapshotData,
diff --git a/Engine/Plugins/VirtualProduction/LevelSnapshots/Source/LevelSnapshots/Private/Data/RestorationEvents/RecreateComponentScope.h b/Engine/Plugins/VirtualProduction/LevelSnapshots/Source/LevelSnapshots/Private/Data/RestorationEvents/RecreateComponentScope.h
--- a/Engine/Plugins/VirtualProduction/LevelSnapshots/Source/LevelSnapshots/Private/Data/RestorationEvents/RecreateComponentScope.h
+++ b/Engine/Plugins/VirtualProduction/LevelSnapshots/Source/LevelSnapshots/Private/Data/RestorationEvents/RecreateComponentScope.h
@@ -3,7 +3,10 @@
 #pragma once
 
 #include "CoreMinimal.h"
+#include "ComponentInstanceDataCache.h"
+#include "Templates/UnrealTemplate.h"
 
+class AActor;
 struct FSubobjectSnapshotData;
 
 /**
diff --git a/Engine/Plugins/VirtualProduction/LevelSnapshots/Source/LevelSnapshots/Private/Data/RestorationEvents/RemoveComponentScope.cpp b/Engine/Plugins/VirtualProduction/LevelSnapshots/Source/LevelSnapshots/Private/Data/RestorationEvents/RemoveComponentScope.cpp
--- a/Engine/Plugins/VirtualProduction/LevelSnapshots/Source/LevelSnapshots/Private/Data/RestorationEvents/RemoveComponentScope.cpp
+++ b/Engine/Plugins/VirtualProduction/LevelSnapshots/Source/LevelSnapshots/Private/Data/RestorationEvents/RemoveComponentScope.cpp
@@ -3,6 +3,7 @@
 #include "Data/RestorationEvents/RemoveComponentScope.h"
 
 #include "LevelSnapshotsModule.h"
+#include "Components/ActorComponent.h"
 
 FRemoveComponentScope::FRemoveComponentScope(UActorComponent* RemovedComponent)
 	:
diff --git a/Engine/Plugins/VirtualProduction/LevelSnapshots/Source/LevelSnapshots/Private/Data/RestorationEvents/RemoveComponentScope.h b/Engine/Plugins/VirtualProduction/LevelSnapshots/Source/LevelSnapshots/Private/Data/RestorationEvents/RemoveComponentScope.h
--- a/Engine/Plugins/VirtualProduction/LevelSnapshots/Source/LevelSnapshots/Private/Data/RestorationEvents/RemoveComponentScope.h
+++ b/Engine/Plugins/VirtualProduction/LevelSnapshots/Source/LevelSnapshots/Private/Data/RestorationEvents/RemoveComponentScope.h
@@ -4,6 +4,9 @@
 
 #include "CoreMinimal.h"
 #include "IRestorationListener.h"
+#include "Templates/UnrealTemplate.h"
+
+class UActorComponent;
 
 /**
  * Convenience type that calls FLevelSnaphshotsModule::OnPreRemoveComponent and FLevelSnaphshotsModule::OnPostRemoveComponent.
diff --git a/Engine/Plugins/VirtualProduction/LevelSnapshots/Source/LevelSnapshots/Private/Data/Util/EquivalenceUtil.cpp b/Engine/Plugins/VirtualProduction/LevelSnapshots/Source/LevelSnapshots/Private/Data/Util/EquivalenceUtil.cpp
--- a/Engine/Plugins/VirtualProduction/LevelSnapshots/Source/LevelSnapshots/Private/Data/Util/EquivalenceUtil.cpp
+++ b/Engine/Plugins/VirtualProduction/LevelSnapshots/Source/LevelSnapshots/Private/Data/Util/EquivalenceUtil.cpp
@@ -7,15 +7,15 @@
 #include "WorldSnapshotData.h"
 
 #include "Components/ActorComponent.h"
+#include "CustomSerialization/CustomObjectSerializationWrapper.h"
 #include "EngineUtils.h"
+#include "GameFramework/Actor.h"
 #include "LevelSnapshotsModule.h"
 #include "PropertyComparisonParams.h"
 #include "PropertyInfoHelpers.h"
 #include "SnapshotCustomVersion.h"
 #include "SnapshotRestorability.h"
-#include "CustomSerialization/CustomObjectSerializationWrapper.h"
-#include "GameFramework/Actor.h"
-#include "Misc/ScopeExit.h"
+#include "UObject/UnrealType.h"
 
 namespace
 {
diff --git a/Engine/Plugins/VirtualProduction/LevelSnapshots/Source/LevelSnapshots/Private/Data/Util/EquivalenceUtil.h b/Engine/Plugins/VirtualProduction/LevelSnapshots/Source/LevelSnapshots/Private/Data/Util/EquivalenceUtil.h
--- a/Engine/Plugins/VirtualProduction/LevelSnapshots/Source/LevelSnapshots/Private/Data/Util/EquivalenceUtil.h
+++ b/Engine/Plugins/VirtualProduction/LevelSnapshots/Source/LevelSnapshots/Private/Data/Util/EquivalenceUtil.h
@@ -5,7 +5,9 @@
 #include "CoreMinimal.h"
 
 class AActor;
+class UActorComponent;
 class UObject;
+
 struct FWorldSnapshotData;
 
 namespace SnapshotUtil
diff --git a/Engine/Plugins/VirtualProduction/LevelSnapshots/Source/LevelSnapshots/Private/PropertySelectionMap.cpp b/Engine/Plugins/VirtualProduction/LevelSnapshots/Source/LevelSnapshots/Private/PropertySelectionMap.cpp
--- a/Engine/Plugins/VirtualProduction/LevelSnapshots/Source/LevelSnapshots/Private/PropertySelectionMap.cpp
+++ b/Engine/Plugins/VirtualProduction/LevelSnapshots/Source/LevelSnapshots/Private/PropertySelectionMap.cpp
@@ -2,7 +2,6 @@
 
 #include "PropertySelectionMap.h"
 
-#include "CustomSubobjectRestorationInfo.h"
 #include "LevelSnapshotsLog.h"
 #include "Data/RestorableObjectSelection.h"
 
diff --git a/Engine/Plugins/VirtualProduction/LevelSnapshots/Source/LevelSnapshots/Private/Tests/ComponentTests.cpp b/Engine/Plugins/VirtualProduction/LevelSnapshots/Source/LevelSnapshots/Private/Tests/ComponentTests.cpp
--- a/Engine/Plugins/VirtualProduction/LevelSnapshots/Source/LevelSnapshots/Private/Tests/ComponentTests.cpp
+++ b/Engine/Plugins/VirtualProduction/LevelSnapshots/Source/LevelSnapshots/Private/Tests/ComponentTests.cpp
@@ -1,13 +1,14 @@
 ﻿// Copyright Epic Games, Inc. All Rights Reserved.
 
-#include "EngineUtils.h"
-#include "PropertySelection.h"
+#include "Data/PropertySelection.h"
 #include "PropertySelectionMap.h"
-#include "SnapshotTestRunner.h"
-#include "SnapshotTestActor.h"
+#include "Tests/SnapshotTestRunner.h"
+#include "Tests/SnapshotTestActor.h"
 
 #include "Components/PointLightComponent.h"
 #include "Components/StaticMeshComponent.h"
+#include "Components/SkeletalMeshComponent.h"
+#include "EngineUtils.h"
 #include "Engine/StaticMesh.h"
 #include "Engine/World.h"
 #include "GameFramework/Actor.h"
diff --git a/Engine/Plugins/VirtualProduction/LevelSnapshots/Source/LevelSnapshots/Public/Data/AddedAndRemovedComponentInfo.h b/Engine/Plugins/VirtualProduction/LevelSnapshots/Source/LevelSnapshots/Public/Data/AddedAndRemovedComponentInfo.h
--- a/Engine/Plugins/VirtualProduction/LevelSnapshots/Source/LevelSnapshots/Public/Data/AddedAndRemovedComponentInfo.h
+++ b/Engine/Plugins/VirtualProduction/LevelSnapshots/Source/LevelSnapshots/Public/Data/AddedAndRemovedComponentInfo.h
@@ -3,7 +3,7 @@
 #pragma once
 
 #include "CoreMinimal.h"
-#include "UObject/WeakObjectPtrTemplates.h"
+#include "UObject/WeakObjectPtr.h"
 
 class UActorComponent;
 
diff --git a/Engine/Plugins/VirtualProduction/LevelSnapshots/Source/LevelSnapshots/Public/Data/CustomSubobjectRestorationInfo.h b/Engine/Plugins/VirtualProduction/LevelSnapshots/Source/LevelSnapshots/Public/Data/CustomSubobjectRestorationInfo.h
--- a/Engine/Plugins/VirtualProduction/LevelSnapshots/Source/LevelSnapshots/Public/Data/CustomSubobjectRestorationInfo.h
+++ b/Engine/Plugins/VirtualProduction/LevelSnapshots/Source/LevelSnapshots/Public/Data/CustomSubobjectRestorationInfo.h
@@ -3,9 +3,7 @@
 #pragma once
 
 #include "CoreMinimal.h"
-#include "UObject/WeakObjectPtrTemplates.h"
-
-class UActorComponent;
+#include "UObject/Object.h"
 
 /* Holds all properties that should be restored for an object. */
 struct LEVELSNAPSHOTS_API FCustomSubobjectRestorationInfo
diff --git a/Engine/Plugins/VirtualProduction/LevelSnapshots/Source/LevelSnapshots/Public/PropertySelectionMap.h b/Engine/Plugins/VirtualProduction/LevelSnapshots/Source/LevelSnapshots/Public/PropertySelectionMap.h
--- a/Engine/Plugins/VirtualProduction/LevelSnapshots/Source/LevelSnapshots/Public/PropertySelectionMap.h
+++ b/Engine/Plugins/VirtualProduction/LevelSnapshots/Source/LevelSnapshots/Public/PropertySelectionMap.h
@@ -7,9 +7,10 @@
 #include "Data/CustomSubobjectRestorationInfo.h"
 #include "Data/PropertySelection.h"
 #include "Data/RestorableObjectSelection.h"
+#include "GameFramework/Actor.h"
 #include "PropertySelectionMap.generated.h"
 
-class AActor;
+class UActorComponent;
 struct FPropertySelectionMap;
 
 /* Binds an object to its selected properties */
diff --git a/Engine/Plugins/VirtualProduction/LevelSnapshots/Source/LevelSnapshots/Public/Restorability/IRestorationListener.h b/Engine/Plugins/VirtualProduction/LevelSnapshots/Source/LevelSnapshots/Public/Restorability/IRestorationListener.h
--- a/Engine/Plugins/VirtualProduction/LevelSnapshots/Source/LevelSnapshots/Public/Restorability/IRestorationListener.h
+++ b/Engine/Plugins/VirtualProduction/LevelSnapshots/Source/LevelSnapshots/Public/Restorability/IRestorationListener.h
@@ -3,6 +3,10 @@
 #pragma once
 
 #include "CoreMinimal.h"
+#include "ComponentInstanceDataCache.h"
+
+class AActor;
+class UActorComponent;
 
 struct FPropertySelection;
 struct FPropertySelectionMap;
diff --git a/Engine/Plugins/VirtualProduction/LevelSnapshots/Source/LevelSnapshotsEditor/Private/Views/Input/LevelSnapshotsEditorInput.h b/Engine/Plugins/VirtualProduction/LevelSnapshots/Source/LevelSnapshotsEditor/Private/Views/Input/LevelSnapshotsEditorInput.h
--- a/Engine/Plugins/VirtualProduction/LevelSnapshots/Source/LevelSnapshotsEditor/Private/Views/Input/LevelSnapshotsEditorInput.h
+++ b/Engine/Plugins/VirtualProduction/LevelSnapshots/Source/LevelSnapshotsEditor/Private/Views/Input/LevelSnapshotsEditorInput.h
@@ -2,9 +2,11 @@
 
 #pragma once
 
+#include "CoreMinimal.h"
 #include "Views/SnapshotEditorViewData.h"
 #include "Templates/SharedPointer.h"
 
+class SWidget;
 class SLevelSnapshotsEditorInput;
 
 struct FAssetData;
diff --git a/Engine/Plugins/VirtualProduction/LevelSnapshots/Source/LevelSnapshotsEditor/Private/Views/Results/LevelSnapshotsEditorResults.h b/Engine/Plugins/VirtualProduction/LevelSnapshots/Source/LevelSnapshotsEditor/Private/Views/Results/LevelSnapshotsEditorResults.h
--- a/Engine/Plugins/VirtualProduction/LevelSnapshots/Source/LevelSnapshotsEditor/Private/Views/Results/LevelSnapshotsEditorResults.h
+++ b/Engine/Plugins/VirtualProduction/LevelSnapshots/Source/LevelSnapshotsEditor/Private/Views/Results/LevelSnapshotsEditorResults.h
@@ -2,10 +2,12 @@
 
 #pragma once
 
+#include "CoreMinimal.h"
 #include "Views/SnapshotEditorViewData.h"
 #include "Templates/SharedPointer.h"
 
 class SLevelSnapshotsEditorResults;
+class SWidget;
 
 class FLevelSnapshotsEditorResults : public TSharedFromThis<FLevelSnapshotsEditorResults>
 {
diff --git a/Engine/Plugins/VirtualProduction/LevelSnapshots/Source/LevelSnapshotsEditor/Private/Widgets/Filter/SSaveAndLoadFilters.cpp b/Engine/Plugins/VirtualProduction/LevelSnapshots/Source/LevelSnapshotsEditor/Private/Widgets/Filter/SSaveAndLoadFilters.cpp
--- a/Engine/Plugins/VirtualProduction/LevelSnapshots/Source/LevelSnapshotsEditor/Private/Widgets/Filter/SSaveAndLoadFilters.cpp
+++ b/Engine/Plugins/VirtualProduction/LevelSnapshots/Source/LevelSnapshotsEditor/Private/Widgets/Filter/SSaveAndLoadFilters.cpp
@@ -14,14 +14,14 @@
 #include "SAssetDropTarget.h"
 #include "Styling/SlateIconFinder.h"
 #include "Widgets/Images/SImage.h"
+#include "Widgets/Input/SComboButton.h"
+#include "Widgets/Layout/SBorder.h"
+#include "Widgets/SBoxPanel.h"
 
 #define LOCTEXT_NAMESPACE "LevelSnapshotsEditor"
 
 namespace
 {
-	constexpr const auto ComboButtonStyle = TEXT("RoundButton");
-	constexpr const auto FontStyle = TEXT("GenericFilters.TextStyle");
-
 	void SaveAs(TWeakObjectPtr<UFilterLoader> FilterLoader)
 	{
 		FilterLoader->SaveAs();
diff --git a/Engine/Plugins/VirtualProduction/LevelSnapshots/Source/LevelSnapshotsEditor/Private/Widgets/SLevelSnapshotsEditorCreationForm.cpp b/Engine/Plugins/VirtualProduction/LevelSnapshots/Source/LevelSnapshotsEditor/Private/Widgets/SLevelSnapshotsEditorCreationForm.cpp
--- a/Engine/Plugins/VirtualProduction/LevelSnapshots/Source/LevelSnapshotsEditor/Private/Widgets/SLevelSnapshotsEditorCreationForm.cpp
+++ b/Engine/Plugins/VirtualProduction/LevelSnapshots/Source/LevelSnapshotsEditor/Private/Widgets/SLevelSnapshotsEditorCreationForm.cpp
@@ -1,6 +1,6 @@
 // Copyright Epic Games, Inc. All Rights Reserved.
 
-#include "SLevelSnapshotsEditorCreationForm.h"
+#include "Widgets/SLevelSnapshotsEditorCreationForm.h"
 
 #include "Data/LevelSnapshotsEditorData.h"
 #include "LevelSnapshotsEditorStyle.h"
@@ -8,6 +8,7 @@
 #include "Settings/LevelSnapshotsEditorDataManagementSettings.h"
 
 #include "EditorStyleSet.h"
+#include "Engine/World.h"
 #include "Framework/Application/SlateApplication.h"
 #include "IDetailsView.h"
 #include "Modules/ModuleManager.h"
@@ -267,8 +268,10 @@ FText SLevelSnapshotsEditorCreationForm::GetNameOverrideText() const
 		return FText::FromString(DataManagementSettingsObjectPtr.Get()->GetNameOverride());
 	}
 
-	return 	ULevelSnapshotsEditorDataManagementSettings::ParseLevelSnapshotsTokensInText(FText::FromString(
-		DataManagementSettingsObjectPtr.Get()->GetNameOverride()), World->GetName());
+	return 	ULevelSnapshotsEditorDataManagementSettings::ParseLevelSnapshotsTokensInText(
+		FText::FromString(DataManagementSettingsObjectPtr.Get()->GetNameOverride()),
+		World->GetName()
+		);
 }
 
 void SLevelSnapshotsEditorCreationForm::SetNameOverrideText(const FText& InNewText, ETextCommit::Type InCommitType)
