Fix compilation error

diff --git a/Engine/Plugins/VirtualProduction/LevelSnapshots/Source/LevelSnapshots/Public/Data/CustomSubobjectRestorationInfo.h b/Engine/Plugins/VirtualProduction/LevelSnapshots/Source/LevelSnapshots/Public/Data/CustomSubobjectRestorationInfo.h
--- a/Engine/Plugins/VirtualProduction/LevelSnapshots/Source/LevelSnapshots/Public/Data/CustomSubobjectRestorationInfo.h
+++ b/Engine/Plugins/VirtualProduction/LevelSnapshots/Source/LevelSnapshots/Public/Data/CustomSubobjectRestorationInfo.h
@@ -4,6 +4,7 @@
 
 #include "CoreMinimal.h"
 #include "UObject/Object.h"
+#include "UObject/WeakObjectPtr.h"
 
 /* Holds all properties that should be restored for an object. */
 struct LEVELSNAPSHOTS_API FCustomSubobjectRestorationInfo
