Fix compile error

diff --git a/Engine/Plugins/VirtualProduction/LevelSnapshots/Source/LevelSnapshots/Private/Archive/ApplySnapshotDataArchiveV2.cpp b/Engine/Plugins/VirtualProduction/LevelSnapshots/Source/LevelSnapshots/Private/Archive/ApplySnapshotDataArchiveV2.cpp
--- a/Engine/Plugins/VirtualProduction/LevelSnapshots/Source/LevelSnapshots/Private/Archive/ApplySnapshotDataArchiveV2.cpp
+++ b/Engine/Plugins/VirtualProduction/LevelSnapshots/Source/LevelSnapshots/Private/Archive/ApplySnapshotDataArchiveV2.cpp
@@ -2,10 +2,13 @@
 
 #include "Archive/ApplySnapshotDataArchiveV2.h"
 
-#include "LevelSnapshotsLog.h"
+#include "Archive/ClassDefaults/ApplyClassDefaulDataArchive.h"
 #include "Data/PropertySelection.h"
-#include "WorldSnapshotData.h"
-#include "ClassDefaults/ApplyClassDefaulDataArchive.h"
+#include "Data/RestorationEvents/ApplySnapshotPropertiesScope.h"
+#include "Data/Util/PropertyUtil.h"
+#include "Data/Util/SnapshotObjectUtil.h"
+#include "Data/WorldSnapshotData.h"
+#include "LevelSnapshotsLog.h"
 
 #include "Components/ActorComponent.h"
 #include "Engine/World.h"
@@ -15,9 +18,6 @@
 #include "Serialization/ObjectWriter.h"
 #include "UObject/TextProperty.h"
 #include "UObject/UnrealType.h"
-#include "RestorationEvents/ApplySnapshotPropertiesScope.h"
-#include "Util/PropertyUtil.h"
-#include "Util/SnapshotObjectUtil.h"
 
 namespace
 {
diff --git a/Engine/Plugins/VirtualProduction/LevelSnapshots/Source/LevelSnapshots/Private/Data/Util/EquivalenceUtil.cpp b/Engine/Plugins/VirtualProduction/LevelSnapshots/Source/LevelSnapshots/Private/Data/Util/EquivalenceUtil.cpp
--- a/Engine/Plugins/VirtualProduction/LevelSnapshots/Source/LevelSnapshots/Private/Data/Util/EquivalenceUtil.cpp
+++ b/Engine/Plugins/VirtualProduction/LevelSnapshots/Source/LevelSnapshots/Private/Data/Util/EquivalenceUtil.cpp
@@ -1,20 +1,21 @@
 ﻿// Copyright Epic Games, Inc. All Rights Reserved.
 
-#include "EquivalenceUtil.h"
+#include "Data/Util/EquivalenceUtil.h"
 
-#include "ApplySnapshotDataArchiveV2.h"
+#include "Archive/ApplySnapshotDataArchiveV2.h"
+#include "Data/WorldSnapshotData.h"
+#include "Data/SnapshotCustomVersion.h"
+#include "Data/CustomSerialization/CustomObjectSerializationWrapper.h"
 #include "LevelSnapshotsLog.h"
-#include "WorldSnapshotData.h"
+#include "LevelSnapshotsModule.h"
+#include "PropertyInfoHelpers.h"
+#include "Restorability/PropertyComparisonParams.h"
+#include "Restorability/SnapshotRestorability.h"
 
 #include "Components/ActorComponent.h"
-#include "CustomSerialization/CustomObjectSerializationWrapper.h"
 #include "EngineUtils.h"
 #include "GameFramework/Actor.h"
-#include "LevelSnapshotsModule.h"
-#include "PropertyComparisonParams.h"
-#include "PropertyInfoHelpers.h"
-#include "SnapshotCustomVersion.h"
-#include "SnapshotRestorability.h"
+#include "UObject/TextProperty.h"
 #include "UObject/UnrealType.h"
 
 namespace
