Fix static analysis

diff --git a/Engine/Plugins/VirtualProduction/LevelSnapshots/Source/LevelSnapshots/Private/Data/Util/SnapshotObjectUtil.cpp b/Engine/Plugins/VirtualProduction/LevelSnapshots/Source/LevelSnapshots/Private/Data/Util/SnapshotObjectUtil.cpp
--- a/Engine/Plugins/VirtualProduction/LevelSnapshots/Source/LevelSnapshots/Private/Data/Util/SnapshotObjectUtil.cpp
+++ b/Engine/Plugins/VirtualProduction/LevelSnapshots/Source/LevelSnapshots/Private/Data/Util/SnapshotObjectUtil.cpp
@@ -130,13 +130,16 @@ namespace
 		}
 
 		// Need to clean rename dead objects otherwise we'll get a name clash when instantiating new object
-		const bool bIsReferencingDeadObject = !IsValid(ResolvedObject) || ResolvedObject->IsUnreachable();
-		const bool bIsReferencingDifferentClass = ResolvedObject && !ResolvedObject->GetClass()->IsChildOf(ExpectedClass);
-		if (bIsReferencingDeadObject || bIsReferencingDifferentClass)
+		if (ResolvedObject)
 		{
-			const FName NewName = MakeUniqueObjectName(ResolvedObject->GetOuter(), ExpectedClass, *ResolvedObject->GetName().Append(TEXT("_TRASH")));
-			ResolvedObject->Rename(*NewName.ToString());
-			ResolvedObject = nullptr;
+			const bool bIsReferencingDeadObject = !IsValid(ResolvedObject) || ResolvedObject->IsUnreachable();
+			const bool bIsReferencingDifferentClass = ResolvedObject && !ResolvedObject->GetClass()->IsChildOf(ExpectedClass);
+			if (bIsReferencingDeadObject || bIsReferencingDifferentClass)
+			{
+				const FName NewName = MakeUniqueObjectName(ResolvedObject->GetOuter(), ExpectedClass, *ResolvedObject->GetName().Append(TEXT("_TRASH")));
+				ResolvedObject->Rename(*NewName.ToString());
+				ResolvedObject = nullptr;
+			}
 		}
 		
 		if (!ensureMsgf(SubobjectData.SnapshotObject.IsValid(), TEXT("Subobject pointer should have been set. Investigate.")))
