- Fixing VP tools toolbar button location

diff --git a/Engine/Plugins/Developer/Concert/MultiUserClient/Source/MultiUserClient/MultiUserClient.Build.cs b/Engine/Plugins/Developer/Concert/MultiUserClient/Source/MultiUserClient/MultiUserClient.Build.cs
--- a/Engine/Plugins/Developer/Concert/MultiUserClient/Source/MultiUserClient/MultiUserClient.Build.cs
+++ b/Engine/Plugins/Developer/Concert/MultiUserClient/Source/MultiUserClient/MultiUserClient.Build.cs
@@ -42,6 +42,7 @@ namespace UnrealBuildTool.Rules
 						"EditorFramework",
 						"UnrealEd",
 						"Sequencer",
+						"ToolMenus",
 						"WorkspaceMenuStructure",
 					}
 				);
diff --git a/Engine/Plugins/Developer/Concert/MultiUserClient/Source/MultiUserClient/Private/MultiUserClientModule.cpp b/Engine/Plugins/Developer/Concert/MultiUserClient/Source/MultiUserClient/Private/MultiUserClientModule.cpp
--- a/Engine/Plugins/Developer/Concert/MultiUserClient/Source/MultiUserClient/Private/MultiUserClientModule.cpp
+++ b/Engine/Plugins/Developer/Concert/MultiUserClient/Source/MultiUserClient/Private/MultiUserClientModule.cpp
@@ -49,6 +49,7 @@
 #include "DetailCategoryBuilder.h"
 #include "DetailWidgetRow.h"
 
+
 #include "Widgets/SConcertBrowser.h"
 #include "Widgets/Docking/SDockTab.h"
 #include "Widgets/Layout/SBox.h"
@@ -56,10 +57,13 @@
 #include "Widgets/Input/SEditableTextBox.h"
 
 #if WITH_EDITOR
+	#include "FileHelpers.h"
 	#include "ISettingsModule.h"
 	#include "ISettingsSection.h"
+	
 	#include "LevelEditor.h"
-	#include "FileHelpers.h"
+	#include "ToolMenuEntry.h"
+	#include "ToolMenus.h"
 	#include "WorkspaceMenuStructure.h"
 	#include "WorkspaceMenuStructureModule.h"
 #endif
@@ -1179,7 +1183,7 @@ private:
 
 			// Register command list
 			FConcertUICommands::Register();
-			TSharedPtr<FUICommandList> CommandList = MakeShared<FUICommandList>();
+			const TSharedPtr<FUICommandList> CommandList = MakeShared<FUICommandList>();
 
 			// Connect to the default server and session
 			CommandList->MapAction(FConcertUICommands::Get().TriggerToolbarButtonCmd,
@@ -1202,31 +1206,29 @@ private:
 				FExecuteAction::CreateRaw(this, &FMultiUserClientModule::LaunchConcertServer)
 			);
 
-			TSharedPtr<FExtender> ToolbarExtender = MakeShared<FExtender>();
-			ToolbarExtender->AddToolBarExtension("Play", EExtensionHook::Before, CommandList, FToolBarExtensionDelegate::CreateLambda([this, CommandList](FToolBarBuilder& ToolbarBuilder)
-			{
-				ToolbarBuilder.AddToolBarButton
-				(
-					FConcertUICommands::Get().TriggerToolbarButtonCmd,
-					NAME_None,
-					TAttribute<FText>::Create([this]() { return GetToolbarButtonIconTitle(); }),
-					TAttribute<FText>::Create([this]() { return GetToolbarButtonTooltip(); }),
-					TAttribute<FSlateIcon>::Create([this]() { return GetToolbarButtonIcon(); })
-				);
+			UToolMenu* Menu = UToolMenus::Get()->ExtendMenu("LevelEditor.LevelEditorToolBar.User");
+			FToolMenuSection& Section = Menu->FindOrAddSection("Concert");
 
-				// Add a simple drop-down menu (no label, no icon for the drop-down button itself) to list
-				ToolbarBuilder.AddComboButton(
-					FUIAction(),
-					FOnGetContent::CreateStatic(&GenerateConcertMenuContent, CommandList.ToSharedRef()),
-					LOCTEXT("ConcertToolbarMenu_Label", "Multi-User Utilities"),
-					LOCTEXT("ConcertToolbarMenu_Tooltip", "Multi-User Commands"),
-					FSlateIcon(),
-					true
-				);
-			}));
+			FToolMenuEntry ConcertButtonEntry = FToolMenuEntry::InitToolBarButton(
+				FConcertUICommands::Get().TriggerToolbarButtonCmd,
+				TAttribute<FText>::Create([this]() { return GetToolbarButtonIconTitle(); }),
+				TAttribute<FText>::Create([this]() { return GetToolbarButtonTooltip(); }),
+				TAttribute<FSlateIcon>::Create([this]() { return GetToolbarButtonIcon(); })
+			);
+			ConcertButtonEntry.SetCommandList(CommandList);
+
+			const FToolMenuEntry ConcertComboEntry = FToolMenuEntry::InitComboButton(
+			"ConcertToolbarMenu",
+			FUIAction(),
+			FOnGetContent::CreateStatic(&GenerateConcertMenuContent, CommandList.ToSharedRef()),
+			LOCTEXT("ConcertToolbarMenu_Label", "Multi-User Utilities"),
+			LOCTEXT("ConcertToolbarMenu_Tooltip", "Multi-User Commands"),
+			FSlateIcon(),
+			true //bInSimpleComboBox
+			);
 
-			LevelEditorModule.GetToolBarExtensibilityManager()->AddExtender(ToolbarExtender);
-			WeakToolbarExtender = ToolbarExtender;
+			Section.AddEntry(ConcertButtonEntry);
+			Section.AddEntry(ConcertComboEntry);
 		}
 	}
 
@@ -1234,10 +1236,12 @@ private:
 	{
 		if (GIsEditor)
 		{
-			FLevelEditorModule& LevelEditorModule = FModuleManager::LoadModuleChecked<FLevelEditorModule>("LevelEditor");
-			LevelEditorModule.GetToolBarExtensibilityManager()->RemoveExtender(WeakToolbarExtender.Pin());
+			if (UObjectInitialized())
+			{
+				UToolMenus::Get()->RemoveSection("LevelEditor.LevelEditorToolBar.User", "Concert");
+			}
+			
 			FConcertUICommands::Unregister();
-			WeakToolbarExtender.Reset();
 		}
 	}
 
@@ -1351,9 +1355,6 @@ private:
 	/** Tracks whether the client has enabled a transport plugin compatible with the server, so they can communicate. */
 	bool bHasRequiredCommunicationPluginEnabled = true;
 
-	/** Pointer to the Toolbar extender. */
-	TWeakPtr<FExtender> WeakToolbarExtender;
-
 	/** UI view and commands on the Concert client workspace. */
 	TSharedPtr<FConcertWorkspaceUI> WorkspaceFrontend;
 
diff --git a/Engine/Plugins/Media/MediaFrameworkUtilities/Source/MediaFrameworkUtilitiesEditor/Private/UI/MediaProfileMenuEntry.cpp b/Engine/Plugins/Media/MediaFrameworkUtilities/Source/MediaFrameworkUtilitiesEditor/Private/UI/MediaProfileMenuEntry.cpp
--- a/Engine/Plugins/Media/MediaFrameworkUtilities/Source/MediaFrameworkUtilitiesEditor/Private/UI/MediaProfileMenuEntry.cpp
+++ b/Engine/Plugins/Media/MediaFrameworkUtilities/Source/MediaFrameworkUtilitiesEditor/Private/UI/MediaProfileMenuEntry.cpp
@@ -83,7 +83,7 @@ struct FMediaProfileMenuEntryImpl
 
 	~FMediaProfileMenuEntryImpl()
 	{
-		if (!IsEngineExitRequested())
+		if (!IsEngineExitRequested() && UObjectInitialized())
 		{
 			UToolMenus::Get()->RemoveSection("LevelEditor.LevelEditorToolBar.User", "MediaProfile");
 		}
diff --git a/Engine/Plugins/VirtualProduction/DMX/DMXEngine/Source/DMXEditor/Private/DMXEditorModule.cpp b/Engine/Plugins/VirtualProduction/DMX/DMXEngine/Source/DMXEditor/Private/DMXEditorModule.cpp
--- a/Engine/Plugins/VirtualProduction/DMX/DMXEngine/Source/DMXEditor/Private/DMXEditorModule.cpp
+++ b/Engine/Plugins/VirtualProduction/DMX/DMXEngine/Source/DMXEditor/Private/DMXEditorModule.cpp
@@ -77,6 +77,12 @@ void FDMXEditorModule::StartupModule()
 void FDMXEditorModule::ShutdownModule()
 {
 	FDMXEditorStyle::Shutdown();
+
+	if(UObjectInitialized())
+	{
+		UToolMenus::Get()->RemoveSection("LevelEditor.LevelEditorToolBar.User", "DMX");
+	}
+	
 	FDMXEditorCommands::Unregister();
 
 	MenuExtensibilityManager.Reset();
@@ -141,7 +147,7 @@ void FDMXEditorModule::BindDMXEditorCommands()
 
 void FDMXEditorModule::ExtendLevelEditorToolbar()
 {
-	UToolMenu* Menu = UToolMenus::Get()->ExtendMenu("LevelEditor.LevelEditorToolBar.ModesToolBar");
+	UToolMenu* Menu = UToolMenus::Get()->ExtendMenu("LevelEditor.LevelEditorToolBar.User");
 
 	FToolMenuSection& Section = Menu->FindOrAddSection("DMX");
 
diff --git a/Engine/Plugins/VirtualProduction/LevelSnapshots/Source/LevelSnapshotsEditor/Private/LevelSnapshotsEditorModule.cpp b/Engine/Plugins/VirtualProduction/LevelSnapshots/Source/LevelSnapshotsEditor/Private/LevelSnapshotsEditorModule.cpp
--- a/Engine/Plugins/VirtualProduction/LevelSnapshots/Source/LevelSnapshotsEditor/Private/LevelSnapshotsEditorModule.cpp
+++ b/Engine/Plugins/VirtualProduction/LevelSnapshots/Source/LevelSnapshotsEditor/Private/LevelSnapshotsEditorModule.cpp
@@ -83,6 +83,8 @@ void FLevelSnapshotsEditorModule::ShutdownModule()
 	{
 		FPropertyEditorModule& PropertyModule = FModuleManager::LoadModuleChecked<FPropertyEditorModule>("PropertyEditor");
 		PropertyModule.UnregisterCustomClassLayout(UNegatableFilter::StaticClass()->GetFName());
+
+		UToolMenus::Get()->RemoveSection("LevelEditor.LevelEditorToolBar.User", "LevelSnapshots");
 	}
 	
 	FLevelSnapshotsEditorCommands::Unregister();
@@ -189,7 +191,7 @@ void FLevelSnapshotsEditorModule::RegisterMenuItem()
 			return;
 		}
 		
-		TSharedRef<FUICommandList> MenuItemCommandList = MakeShareable(new FUICommandList);
+		const TSharedRef<FUICommandList> MenuItemCommandList = MakeShared<FUICommandList>();
 
 		MenuItemCommandList->MapAction(
 			FLevelSnapshotsEditorCommands::Get().OpenLevelSnapshotsEditorMenuItem,
@@ -267,25 +269,37 @@ void FLevelSnapshotsEditorModule::RegisterEditorToolbar()
 		return;
 	}
 	
-	// Get the Level Editor so we can insert our combo button into the Level Editor's toolbar
-	FLevelEditorModule& LevelEditorModule = FModuleManager::LoadModuleChecked<FLevelEditorModule>("LevelEditor");
-	
 	MapEditorToolbarActions();
 
-	// Create a toolbar extension instance that will insert our toolbar button after the 'Settings' horizontal box in the toolbar
-	TSharedPtr<FExtender> NewToolbarExtender = MakeShareable(new FExtender);
-	NewToolbarExtender->AddToolBarExtension("Settings",
-		EExtensionHook::After,
-		EditorToolbarButtonCommandList,
-		FToolBarExtensionDelegate::CreateRaw(this, &FLevelSnapshotsEditorModule::CreateEditorToolbarButton));
-	
-	// Now insert the button into the main toolbar using the extension
-	LevelEditorModule.GetToolBarExtensibilityManager()->AddExtender(NewToolbarExtender);
+	UToolMenu* Menu = UToolMenus::Get()->ExtendMenu("LevelEditor.LevelEditorToolBar.User");
+	FToolMenuSection& Section = Menu->FindOrAddSection("LevelSnapshots");
+
+	FToolMenuEntry LevelSnapshotsButtonEntry = FToolMenuEntry::InitToolBarButton(
+		"TakeSnapshotAction",
+		FUIAction(FExecuteAction::CreateStatic(&SnapshotEditor::TakeSnapshotWithOptionalForm)),
+		NSLOCTEXT("LevelSnapshots", "LevelSnapshots", "Level Snapshots"), // Set Text under image
+		NSLOCTEXT("LevelSnapshots", "LevelSnapshotsToolbarButtonTooltip", "Take snapshot with optional form"), //  Set tooltip
+		FSlateIcon(FLevelSnapshotsEditorStyle::GetStyleSetName(), "LevelSnapshots.ToolbarButton", "LevelSnapshots.ToolbarButton.Small") // Set image
+	);
+	LevelSnapshotsButtonEntry.SetCommandList(EditorToolbarButtonCommandList);
+
+	 FToolMenuEntry LevelSnapshotsComboEntry = FToolMenuEntry::InitComboButton(
+		"LevelSnapshotsMenu",
+		FUIAction(),
+		FOnGetContent::CreateRaw(this, &FLevelSnapshotsEditorModule::FillEditorToolbarComboButtonMenuOptions, EditorToolbarButtonCommandList),
+		NSLOCTEXT("LevelSnapshots", "LevelSnapshotsOptions_Label", "Level Snapshots Options"), // Set text seen when the Level Editor Toolbar is truncated and the flyout is clicked
+		NSLOCTEXT("LevelSnapshots", "LevelSnapshotsToolbarComboButtonTooltip", "Open Level Snapshots Options"), //  Set tooltip
+		FSlateIcon(),
+		true //bInSimpleComboBox
+	);
+
+	Section.AddEntry(LevelSnapshotsButtonEntry);
+	Section.AddEntry(LevelSnapshotsComboEntry);
 }
 
 void FLevelSnapshotsEditorModule::MapEditorToolbarActions()
 {
-	EditorToolbarButtonCommandList = MakeShareable(new FUICommandList);
+	EditorToolbarButtonCommandList = MakeShared<FUICommandList>();
 
 	EditorToolbarButtonCommandList->MapAction(
 		FLevelSnapshotsEditorCommands::Get().UseCreationFormToggle,
@@ -307,24 +321,6 @@ void FLevelSnapshotsEditorModule::MapEditorToolbarActions()
 	);
 }
 
-void FLevelSnapshotsEditorModule::CreateEditorToolbarButton(FToolBarBuilder& Builder)
-{
-	Builder.AddToolBarButton(
-		FUIAction(FExecuteAction::CreateStatic(&SnapshotEditor::TakeSnapshotWithOptionalForm)),
-		NAME_None,
-		NSLOCTEXT("LevelSnapshots", "LevelSnapshots", "Level Snapshots"), // Set Text under image
-		NSLOCTEXT("LevelSnapshots", "LevelSnapshotsToolbarButtonTooltip", "Take snapshot with optional form"), //  Set tooltip
-		FSlateIcon(FLevelSnapshotsEditorStyle::GetStyleSetName(), "LevelSnapshots.ToolbarButton", "LevelSnapshots.ToolbarButton.Small") // Set image
-	);
-	
-	Builder.AddComboButton(
-		FUIAction(),
-		FOnGetContent::CreateRaw(this, &FLevelSnapshotsEditorModule::FillEditorToolbarComboButtonMenuOptions, EditorToolbarButtonCommandList), // Add combo button subcommands menu
-		NSLOCTEXT("LevelSnapshots", "LevelSnapshotsOptions_Label", "Level Snapshots Options"), // Set text seen when the Level Editor Toolbar is truncated and the flyout is clicked
-		NSLOCTEXT("LevelSnapshots", "LevelSnapshotsToolbarComboButtonTooltip", "Open Level Snapshots Options"), FSlateIcon(), true //  Set tooltip
-	); 
-}
-
 TSharedRef<SWidget> FLevelSnapshotsEditorModule::FillEditorToolbarComboButtonMenuOptions(TSharedPtr<class FUICommandList> Commands)
 {
 	// Create FMenuBuilder instance for the commands we created
diff --git a/Engine/Plugins/VirtualProduction/LevelSnapshots/Source/LevelSnapshotsEditor/Private/LevelSnapshotsEditorModule.h b/Engine/Plugins/VirtualProduction/LevelSnapshots/Source/LevelSnapshotsEditor/Private/LevelSnapshotsEditorModule.h
--- a/Engine/Plugins/VirtualProduction/LevelSnapshots/Source/LevelSnapshotsEditor/Private/LevelSnapshotsEditorModule.h
+++ b/Engine/Plugins/VirtualProduction/LevelSnapshots/Source/LevelSnapshotsEditor/Private/LevelSnapshotsEditorModule.h
@@ -47,7 +47,6 @@ private:
 	
 	void RegisterEditorToolbar();
 	void MapEditorToolbarActions();
-	void CreateEditorToolbarButton(FToolBarBuilder& Builder);
 	TSharedRef<SWidget> FillEditorToolbarComboButtonMenuOptions(TSharedPtr<class FUICommandList> Commands);
 
 	void OpenSnapshotsEditor();
diff --git a/Engine/Plugins/VirtualProduction/Switchboard/Source/SwitchboardEditor/Private/SwitchboardMenuEntry.cpp b/Engine/Plugins/VirtualProduction/Switchboard/Source/SwitchboardEditor/Private/SwitchboardMenuEntry.cpp
--- a/Engine/Plugins/VirtualProduction/Switchboard/Source/SwitchboardEditor/Private/SwitchboardMenuEntry.cpp
+++ b/Engine/Plugins/VirtualProduction/Switchboard/Source/SwitchboardEditor/Private/SwitchboardMenuEntry.cpp
@@ -15,6 +15,7 @@
 #include "SwitchboardEditorModule.h"
 #include "SwitchboardEditorSettings.h"
 #include "SwitchboardEditorStyle.h"
+#include "ToolMenus.h"
 
 #define LOCTEXT_NAMESPACE "SwitchboardEditor"
 
@@ -40,51 +41,42 @@ struct FSwitchboardMenuEntryImpl
 	FSwitchboardMenuEntryImpl()
 	{
 		FSwitchboardUICommands::Register();
-		TSharedPtr<FUICommandList> Actions = MakeShareable(new FUICommandList);
+		const TSharedPtr<FUICommandList> Actions = MakeShared<FUICommandList>();
 
 		Actions->MapAction(FSwitchboardUICommands::Get().TriggerToolbarButtonCmd,
 			FExecuteAction::CreateRaw(this, &FSwitchboardMenuEntryImpl::OnLaunchSwitchboardClicked),
 			FCanExecuteAction::CreateRaw(this, &FSwitchboardMenuEntryImpl::IsToolbarButtonEnabled));
 
-		ToolBarExtender = MakeShareable(new FExtender);
-		ToolBarExtender->AddToolBarExtension("Game", EExtensionHook::After, Actions, FToolBarExtensionDelegate::CreateLambda([this, Actions](FToolBarBuilder& ToolbarBuilder)
-		{
-			ToolbarBuilder.BeginSection("Switchboard");
-			{
-				ToolbarBuilder.AddToolBarButton
-				(
-					FSwitchboardUICommands::Get().TriggerToolbarButtonCmd,
-					NAME_None,
-					TAttribute<FText>::Create([this]() { return LOCTEXT("LaunchSwitchboard", "Switchboard"); }),
-					TAttribute<FText>::Create([this]() { return LOCTEXT("LaunchSwitchboardTooltip", "Launch Switchboard"); }),
-					TAttribute<FSlateIcon>::Create([this]() { return GetToolbarButtonIcon(); })
-				);
-				ToolbarBuilder.AddComboButton
-				(
-					FUIAction(),
-					FOnGetContent::CreateRaw(this, &FSwitchboardMenuEntryImpl::CreateListenerEntries),
-					TAttribute<FText>::Create([this]() { return LOCTEXT("LaunchSwitchboard", "Switchboard"); }),
-					LOCTEXT("SwitchboardTooltip", "Actions related to the SwitchboardListener"),
-					FSlateIcon(),
-					true
-				);
-			}
-			ToolbarBuilder.EndSection();
-		}));
-
-		FLevelEditorModule& LevelEditorModule = FModuleManager::LoadModuleChecked<FLevelEditorModule>("LevelEditor");
-		LevelEditorModule.GetToolBarExtensibilityManager()->AddExtender(ToolBarExtender);
+		UToolMenu* Menu = UToolMenus::Get()->ExtendMenu("LevelEditor.LevelEditorToolBar.User");
+		FToolMenuSection& Section = Menu->FindOrAddSection("Switchboard");
+
+		FToolMenuEntry SwitchboardButtonEntry = FToolMenuEntry::InitToolBarButton(
+			FSwitchboardUICommands::Get().TriggerToolbarButtonCmd,
+			TAttribute<FText>::Create([this]() { return LOCTEXT("LaunchSwitchboard", "Switchboard"); }),
+			TAttribute<FText>::Create([this]() { return LOCTEXT("LaunchSwitchboardTooltip", "Launch Switchboard"); }),
+			TAttribute<FSlateIcon>::Create([this]() { return GetToolbarButtonIcon(); })
+		);
+		SwitchboardButtonEntry.SetCommandList(Actions);
+
+		const FToolMenuEntry SwitchboardComboEntry = FToolMenuEntry::InitComboButton(
+			"SwitchboardMenu",
+			FUIAction(),
+			FOnGetContent::CreateRaw(this, &FSwitchboardMenuEntryImpl::CreateListenerEntries),
+			TAttribute<FText>::Create([this]() { return LOCTEXT("LaunchSwitchboard", "Switchboard"); }),
+			LOCTEXT("SwitchboardTooltip", "Actions related to the SwitchboardListener"),
+			FSlateIcon(),
+			true //bInSimpleComboBox
+		);
+
+		Section.AddEntry(SwitchboardButtonEntry);
+		Section.AddEntry(SwitchboardComboEntry);
 	}
 
 	~FSwitchboardMenuEntryImpl()
 	{
-		if (!IsEngineExitRequested() && ToolBarExtender.IsValid())
+		if (!IsEngineExitRequested() && UObjectInitialized())
 		{
-			FLevelEditorModule* LevelEditorModule = FModuleManager::GetModulePtr<FLevelEditorModule>("LevelEditor");
-			if (LevelEditorModule)
-			{
-				LevelEditorModule->GetToolBarExtensibilityManager()->RemoveExtender(ToolBarExtender);
-			}
+			UToolMenus::Get()->RemoveSection("LevelEditor.LevelEditorToolBar.User", "Switchboard");
 		}
 	}
 
@@ -240,7 +232,6 @@ struct FSwitchboardMenuEntryImpl
 	}
 
 private:
-	TSharedPtr<FExtender> ToolBarExtender;
 	ECheckBoxState CachedAutolaunchCheckState;
 
 public:
diff --git a/Engine/Plugins/VirtualProduction/Switchboard/Source/SwitchboardEditor/SwitchboardEditor.build.cs b/Engine/Plugins/VirtualProduction/Switchboard/Source/SwitchboardEditor/SwitchboardEditor.build.cs
--- a/Engine/Plugins/VirtualProduction/Switchboard/Source/SwitchboardEditor/SwitchboardEditor.build.cs
+++ b/Engine/Plugins/VirtualProduction/Switchboard/Source/SwitchboardEditor/SwitchboardEditor.build.cs
@@ -19,6 +19,7 @@ public class SwitchboardEditor : ModuleRules
 				"Settings",
 				"Slate",
 				"SlateCore",
+				"ToolMenus",
 				"UnrealEd",
 			});
 	}
