Patch for for: //UE5/Main - Build DDC Win64 - Assertion occurs with Level Snapshots

#[fyi] Jason.Walter

diff --git a/Engine/Plugins/VirtualProduction/LevelSnapshots/Source/LevelSnapshotsEditor/Private/LevelSnapshotsEditorModule.cpp b/Engine/Plugins/VirtualProduction/LevelSnapshots/Source/LevelSnapshotsEditor/Private/LevelSnapshotsEditorModule.cpp
--- a/Engine/Plugins/VirtualProduction/LevelSnapshots/Source/LevelSnapshotsEditor/Private/LevelSnapshotsEditorModule.cpp
+++ b/Engine/Plugins/VirtualProduction/LevelSnapshots/Source/LevelSnapshotsEditor/Private/LevelSnapshotsEditorModule.cpp
@@ -124,7 +124,8 @@ void FLevelSnapshotsEditorModule::PostEngineInit()
 
 	if (RegisterProjectSettings() && ProjectSettingsObjectPtr->bEnableLevelSnapshotsToolbarButton)
 	{
-		RegisterEditorToolbar();
+		// todo: Skip toolbar for UE5 for now; we will reimplement toolbar at a later date
+		//RegisterEditorToolbar();
 	}
 }
 
