Fix transient actors being allowed

#[fyi] Jason.Walter
Fix nullptr crash

Fix collision wrongly being marked as changed

-Added a delegate to PropertyRowGenerator that will be called instead of ValidatePropertyNodes every Tick if it's set. Exposed a function in IPropertyRowGenerator interface to set this on the implementation. This is done because ValidatePropertyNodes is not needed for every implementation and is heavy to run, especially when many PRG instances exist at once.

#[fyi] Jason.Walter
#[fyi] Sebastian.Nordgren

Functional changes:
- Fix parented actors restoring incorrect transform
- Fix parented actors losing their attach parent
- Properly prevent world objects being copy pasted from snapshot to editor world when restoring

Refactor code:
- Prepare code for supporting restoring subobjects and add TODOs for what needs to be done
- Make it more clear whether a reference is being restored for the editor or snapshot world

Fix FText display string not getting serialized

Temp fix for snapshots not serializing correctly

#[fyi] Jason.Walter
Results:
-Removed some unused includes and forward declarations
-Moved FPropertyHandleHierarchy to the script scope rather than a local struct scope to make it a member of FLevelSnapshotsEditorResultsRow
-Added Lambda to PropertyRowGenerator registration that returns true to avoid heavy ticking.
-Removed WidgetTypeCustomization logic as it was no longer needed after getting customized widgets from PRG
-Added fix for actor groups being created when there are no properties inside
-Added fix for component/subobject groups not being created because they have no direct properties but nested groups do have properties. Now they are created when there are direct properties or any child groups with direct properties.
-Now when building a selection set on apply to world, unchecked properties are removed from a copied selection set rather than having checked properties added to a new one.
-Added a combo button with eyeball iconography beside the search bar. This is where collapse function will now live as well as a few options to determine what to show. Besides Collapse All, these options are not yet working because I need to confirm with Ryan and Tom which options we want to keep before refactoring to accommodate them.
-The Tree View generation methods now return a bool determined by whether or not the list of children has more than 0 members. If it returns false, the header row will be destroyed.
-Now we generate customized widgets when a row is scrolled into view. it allows us to do less work upfront.
-When flushing references, we check first if the references are valid. This saves a few ticks per row.
-Customized widgets are not put inside of Invalidation Panels anymore because some of them already contain panels of their own and Slate doesn't like the invalidation panel recursion.

PropertyInfoHelpers:
-Updated IsPropertyInContainer to return IsPropertyInCollection(Property) or IsPropertyInStruct(Property) rather than doing double work.
-Updated IsPropertyInStruct(Property) to check to see if a property's owner is a scriptstruct as well. There are occasions in which a property will have FStructProperty as its parent or a USCriptStruct object as its parent. Now we check both if necessary to  ascertain if a property is inside of a struct.

PropertySelectionMap:
-Added a GetKeyCount function to simply get the number of keys. Useful for building results view.

#[fyi] Jason.Walter

Fix default (sub)object references getting serialized

#[fyi] Jason.Walter
Adjust how FText is serialized

diff --git a/Engine/Plugins/VirtualProduction/LevelSnapshots/Source/LevelSnapshots/Private/ApplySnapshotFilter.cpp b/Engine/Plugins/VirtualProduction/LevelSnapshots/Source/LevelSnapshots/Private/ApplySnapshotFilter.cpp
--- a/Engine/Plugins/VirtualProduction/LevelSnapshots/Source/LevelSnapshots/Private/ApplySnapshotFilter.cpp
+++ b/Engine/Plugins/VirtualProduction/LevelSnapshots/Source/LevelSnapshots/Private/ApplySnapshotFilter.cpp
@@ -190,13 +190,22 @@ void FApplySnapshotFilter::AnalyseRootProperties(FPropertyContainerContext& Cont
 		// Ask external modules about the property
 		const FPropertyComparisonParams Params { ContainerContext.RootClass, *FieldIt, ContainerContext.SnapshotContainer, ContainerContext.WorldContainer, SnapshotObject, WorldObject, DeserializedSnapshotActor, WorldActor} ;
 		const IPropertyComparer::EPropertyComparison ComparisonResult = Module.ShouldConsiderPropertyEqual(PropertyComparers, Params);
-		if (ComparisonResult == IPropertyComparer::EPropertyComparison::TreatEqual)
+
+		bool bSkipEqualityTest = false;
+		switch (ComparisonResult)
 		{
+		case IPropertyComparer::EPropertyComparison::TreatEqual:
 			continue;
-		} 
+
+		case IPropertyComparer::EPropertyComparison::TreatUnequal:
+			bSkipEqualityTest = true;
+
+		default:
+			break;
+		}
 		
-		const ECheckSubproperties CheckSubpropertyBehaviour = AnalyseProperty(ContainerContext, *FieldIt);
-		if (CheckSubpropertyBehaviour == ECheckSubproperties::CheckSubproperties)
+		const ECheckSubproperties CheckSubpropertyBehaviour = AnalyseProperty(ContainerContext, *FieldIt, bSkipEqualityTest);
+		if (!bSkipEqualityTest && CheckSubpropertyBehaviour == ECheckSubproperties::CheckSubproperties)
 		{
 			HandleStructProperties(ContainerContext, *FieldIt);
 			// TODO: To analyse subobjects, you'd add another utility function here
@@ -225,12 +234,12 @@ void FApplySnapshotFilter::HandleStructProperties(FPropertyContainerContext& Con
 	}
 }
 
-FApplySnapshotFilter::ECheckSubproperties FApplySnapshotFilter::AnalyseProperty(FPropertyContainerContext& ContainerContext, FProperty* PropertyInCommon)
+FApplySnapshotFilter::ECheckSubproperties FApplySnapshotFilter::AnalyseProperty(FPropertyContainerContext& ContainerContext, FProperty* PropertyInCommon, bool bSkipEqualityTest)
 {
 	check(ContainerContext.WorldContainer);
 	check(ContainerContext.SnapshotContainer); 
 
-	if (!bAllowUnchangedProperties && Snapshot->AreSnapshotAndOriginalPropertiesEquivalent(PropertyInCommon, ContainerContext.SnapshotContainer, ContainerContext.WorldContainer, DeserializedSnapshotActor, WorldActor))
+	if (!bSkipEqualityTest && !bAllowUnchangedProperties && Snapshot->AreSnapshotAndOriginalPropertiesEquivalent(PropertyInCommon, ContainerContext.SnapshotContainer, ContainerContext.WorldContainer, DeserializedSnapshotActor, WorldActor))
 	{
 		return SkipSubproperties;
 	}
diff --git a/Engine/Plugins/VirtualProduction/LevelSnapshots/Source/LevelSnapshots/Private/ApplySnapshotFilter.h b/Engine/Plugins/VirtualProduction/LevelSnapshots/Source/LevelSnapshots/Private/ApplySnapshotFilter.h
--- a/Engine/Plugins/VirtualProduction/LevelSnapshots/Source/LevelSnapshots/Private/ApplySnapshotFilter.h
+++ b/Engine/Plugins/VirtualProduction/LevelSnapshots/Source/LevelSnapshots/Private/ApplySnapshotFilter.h
@@ -70,7 +70,7 @@ private:
 		CheckSubproperties,
         SkipSubproperties
     }; 
-	ECheckSubproperties AnalyseProperty(FPropertyContainerContext& ContainerContext, FProperty* PropertyInCommon);
+	ECheckSubproperties AnalyseProperty(FPropertyContainerContext& ContainerContext, FProperty* PropertyInCommon, bool bSkipEqualityTest = false);
 
 	
 	ULevelSnapshot* Snapshot;
diff --git a/Engine/Plugins/VirtualProduction/LevelSnapshots/Source/LevelSnapshots/Private/Archive/ApplySnapshotDataArchiveV2.cpp b/Engine/Plugins/VirtualProduction/LevelSnapshots/Source/LevelSnapshots/Private/Archive/ApplySnapshotDataArchiveV2.cpp
--- a/Engine/Plugins/VirtualProduction/LevelSnapshots/Source/LevelSnapshots/Private/Archive/ApplySnapshotDataArchiveV2.cpp
+++ b/Engine/Plugins/VirtualProduction/LevelSnapshots/Source/LevelSnapshots/Private/Archive/ApplySnapshotDataArchiveV2.cpp
@@ -2,14 +2,20 @@
 
 #include "Archive/ApplySnapshotDataArchiveV2.h"
 
+#include "LevelSnapshotsLog.h"
 #include "Data/PropertySelection.h"
 #include "LevelSnapshotsStats.h"
+#include "Archive/ClassDefaults/TakeClassDefaultObjectSnapshotArchive.h"
 #include "WorldSnapshotData.h"
+#include "ClassDefaults/ApplyClassDefaulDataArchive.h"
 
 #include "Components/ActorComponent.h"
+#include "Engine/World.h"
 #include "GameFramework/Actor.h"
 #include "Serialization/ObjectReader.h"
 #include "Serialization/ObjectWriter.h"
+#include "UObject/TextProperty.h"
+#include "UObject/UnrealType.h"
 
 namespace
 {
@@ -18,67 +24,143 @@ namespace
 		using Super = FObjectWriter;
 
 		const FPropertySelection& PropertiesToSerialize;
+		const UObject* SnapshotObject;
+
+		bool IsWorldObjectProperty(const FProperty* InProperty) const
+		{
+			const FObjectPropertyBase* ObjectProperty = CastField<FObjectPropertyBase>(InProperty);
+			if (!ObjectProperty)
+			{
+				return false;
+			}
+
+			const bool bIsMarkedAsSubobject = InProperty->HasAnyPropertyFlags(CPF_InstancedReference | CPF_ContainsInstancedReference | CPF_PersistentInstance);
+			const bool bIsActorOrComponentPtr = ObjectProperty->PropertyClass->IsChildOf(AActor::StaticClass()) || ObjectProperty->PropertyClass->IsChildOf(UActorComponent::StaticClass());
+			if (bIsMarkedAsSubobject || bIsActorOrComponentPtr)
+			{
+				return true;
+			}
+
+			const FArchiveSerializedPropertyChain* PropertyChain = GetSerializedPropertyChain();
+			const void* ContainerPtr = SnapshotObject;
+			for (int32 i = 0; PropertyChain && i < PropertyChain->GetNumProperties(); ++i)
+			{
+				ContainerPtr = PropertyChain->GetPropertyFromRoot(i)->ContainerPtrToValuePtr<void>(ContainerPtr);
+			}
+
+			const void* LeafValuePtr = InProperty->ContainerPtrToValuePtr<void>(ContainerPtr);
+			const UObject* ContainedPtr = ObjectProperty->GetObjectPropertyValue(LeafValuePtr);
+			return ContainerPtr ? ContainedPtr->IsInA(UWorld::StaticClass()) : false;
+		}
 				
 	public:
 
-		FCopyProperties(TArray<uint8>& SaveLocation, const FPropertySelection& PropertiesToSerialize)
+		mutable TSet<const FTextProperty*> TextProperties;
+
+		FCopyProperties(UObject* SnapshotObject, TArray<uint8>& SaveLocation, const FPropertySelection& PropertiesToSerialize)
 			:
 			Super(SaveLocation),
-			PropertiesToSerialize(PropertiesToSerialize)
+			PropertiesToSerialize(PropertiesToSerialize),
+			SnapshotObject(SnapshotObject),
+			TextProperties(TextProperties)
 		{
 			ArNoDelta = true;
 		}
 
 		virtual bool ShouldSkipProperty(const FProperty* InProperty) const override
 		{
-			if (const FObjectPropertyBase* ObjectProperty = CastField<FObjectPropertyBase>(InProperty))
+			// Do not copy object reference properties that have a world as outer: They will not be valid when copied from snapshot to editor world.
+			// Hence, we only allow object references to external assets, e.g. UMaterials or UDataAssets
+			if (IsWorldObjectProperty(InProperty))
+			{
+				return true;
+			}
+
+			const bool bIsPropertyAllowed = PropertiesToSerialize.ShouldSerializeProperty(GetSerializedPropertyChain(), InProperty);
+			if (bIsPropertyAllowed)
 			{
-				const bool bIsSubobject = InProperty->HasAnyPropertyFlags(CPF_InstancedReference | CPF_ContainsInstancedReference | CPF_PersistentInstance);
-				const bool bIsWorldPointer = ObjectProperty->PropertyClass->IsInA(AActor::StaticClass()) || ObjectProperty->PropertyClass->IsInA(UActorComponent::StaticClass());
-				if (bIsSubobject || bIsWorldPointer)
+				if (const FTextProperty* TextProperty = CastField<FTextProperty>(InProperty))
 				{
+					TextProperties.Add(TextProperty);
 					return false;
 				}
 			}
-			return !PropertiesToSerialize.ShouldSerializeProperty(GetSerializedPropertyChain(), InProperty);
-		}
+			return !bIsPropertyAllowed;
+		}		
 	};
-}
 
-void FApplySnapshotDataArchiveV2::ApplyToExistingWorldObject(FObjectSnapshotData& InObjectData, FWorldSnapshotData& InSharedData, UObject* InOriginalObject, UObject* InDeserializedVersion, const FPropertySelection& InSelectionSet)
-{
-	// ObjectData only contains properties that were different from the CDO at the time of saving. This archive may skip many properties.
-	// Hence, we serialise in two steps:
-
-	
-	// Step 1: Serialise  properties that were different from CDO at time of snapshotting and that are still different from CDO
-	// Most UObject are handled here: world pointers (to Actors and Components) and most asset pointers (material, data assets, etc.). UObject references are a special case: object references were saved even if the value was the same as the CDO. This is done to avoid copy-pasting errors in step 2.
-		// - Subobject references are handled here
-		// - References to other actors in the world are handled here
-	FApplySnapshotDataArchiveV2 ApplySavedData(InObjectData, InSharedData, InOriginalObject, InSelectionSet);
-	InOriginalObject->Serialize(ApplySavedData);
+	class FSerializeTextProperties : public FApplyClassDefaulDataArchive
+	{
+		TSet<const FTextProperty*>& TextProperties;
+	public:
+		
+		FSerializeTextProperties(TSet<const FTextProperty*>& InTextProperties, FObjectSnapshotData& InObjectData, FWorldSnapshotData& InSharedData, UObject* InSerializedObject)
+			:
+			FApplyClassDefaulDataArchive(InObjectData, InSharedData, InSerializedObject, ESerialisationMode::RestoringChangedDefaults),
+			TextProperties(InTextProperties)
+		{}
 
+		virtual bool ShouldSkipProperty(const FProperty* InProperty) const override
+		{
+			if (const FTextProperty* TextProperty = CastField<FTextProperty>(InProperty))
+			{
+				return !TextProperties.Contains(TextProperty);
+			}
+			return true;
+		}
+	};
 
-	
-	// Step 2: Serialise any remaining properties that were not covered: properties that were the equal to the CDO value when the snapshot was taken, but now are different from the CDO.
-	// For this step, we indirectly use the CDO values saved in the snapshot: we copy over all remaining properties from the deserialized version.
-	// The only UObject properties handled here are external references to e.g. materials or data assets. Example: the reference in CDO when snapshot was taken was different from the reference in current CDO.
-	const FPropertySelection& PropertiesLeftToSerialise = ApplySavedData.PropertiesLeftToSerialize;
-	if (!PropertiesLeftToSerialise.IsEmpty())
+	TSet<const FTextProperty*> CopyPastePropertiesDifferentInCDO(const FPropertySelection& PropertiesLeftToSerialise, UObject* InOriginalObject, UObject* InDeserializedVersion)
 	{
-		TArray<uint8> CopiedPropertyData;
-		FCopyProperties CopySimpleProperties(CopiedPropertyData, PropertiesLeftToSerialise);
-		InDeserializedVersion->Serialize(CopySimpleProperties);
+		if (!PropertiesLeftToSerialise.IsEmpty())
+		{
+			TArray<uint8> CopiedPropertyData;
+			FCopyProperties CopySimpleProperties(InDeserializedVersion, CopiedPropertyData, PropertiesLeftToSerialise);
+			InDeserializedVersion->Serialize(CopySimpleProperties);
 		
-		FObjectReader PasteSimpleProperties(InOriginalObject, CopiedPropertyData);
+			FObjectReader PasteSimpleProperties(InOriginalObject, CopiedPropertyData);
+			return CopySimpleProperties.TextProperties;
+		}
+		return {};
+	}
+
+	void FixUpTextPropertiesDifferentInCDO(TSet<const FTextProperty*> TextProperties, FWorldSnapshotData& InSharedData, UObject* InOriginalObject)
+	{
+		const bool bAreTextPropertiesLeft = TextProperties.Num() > 0;
+		if (bAreTextPropertiesLeft)
+		{
+			FObjectSnapshotData* ClassDefaults = InSharedData.GetSerializedClassDefaults(InOriginalObject->GetClass()).Get(nullptr);
+			if (ClassDefaults)
+			{
+				FSerializeTextProperties SerializeTextProperties(TextProperties, *ClassDefaults, InSharedData, InOriginalObject);
+				InOriginalObject->Serialize(SerializeTextProperties);
+			}
+			else
+			{
+				UE_LOG(LogLevelSnapshots, Warning, TEXT("%d text properties have changed in class defaults since snapshot was taken but cannot be restored."), TextProperties.Num());
+			}
+		}
 	}
 }
 
-void FApplySnapshotDataArchiveV2::ApplyToRecreatedWorldObject(FObjectSnapshotData& InObjectData, FWorldSnapshotData& InSharedData, UObject* InOriginalObject, UObject* InDeserializedVersion)
+void FApplySnapshotDataArchiveV2::ApplyToExistingEditorWorldObject(FObjectSnapshotData& InObjectData, FWorldSnapshotData& InSharedData, UObject* InOriginalObject, UObject* InDeserializedVersion, const FPropertySelectionMap& InSelectionMapForResolvingSubobjects, const FPropertySelection& InSelectionSet)
+{	
+	// Step 1: Serialise  properties that were different from CDO at time of snapshotting and that are still different from CDO
+	FApplySnapshotDataArchiveV2 ApplySavedData(InObjectData, InSharedData, InOriginalObject, InSelectionMapForResolvingSubobjects, &InSelectionSet);
+	InOriginalObject->Serialize(ApplySavedData);
+	
+	// Step 2: Serialise any remaining properties that were not covered: properties that were equal to the CDO value when the snapshot was taken but now are different from the CDO.
+	TSet<const FTextProperty*> TextProperties = CopyPastePropertiesDifferentInCDO(ApplySavedData.PropertiesLeftToSerialize, InOriginalObject, InDeserializedVersion);
+
+	// Step 3: Serialize FText properties that have changed in CDO since snapshot was taken. 
+	FixUpTextPropertiesDifferentInCDO(MoveTemp(TextProperties), InSharedData, InOriginalObject);
+}
+
+void FApplySnapshotDataArchiveV2::ApplyToRecreatedEditorWorldObject(FObjectSnapshotData& InObjectData, FWorldSnapshotData& InSharedData, UObject* InOriginalObject, UObject* InDeserializedVersion, const FPropertySelectionMap& InSelectionMapForResolvingSubobjects)
 {
 	// Apply all properties that we saved into the target actor.
 	// We assume that InOriginalObject was already created with the snapshot CDO as template: we do not need Step 2 from ApplyToExistingWorldObject.
-	FApplySnapshotDataArchiveV2 ApplySavedData(InObjectData, InSharedData, InOriginalObject);
+	FApplySnapshotDataArchiveV2 ApplySavedData(InObjectData, InSharedData, InOriginalObject, InSelectionMapForResolvingSubobjects, {});
 	InOriginalObject->Serialize(ApplySavedData);
 }
 
@@ -88,10 +170,10 @@ bool FApplySnapshotDataArchiveV2::ShouldSkipProperty(const FProperty* InProperty
 	
 	bool bShouldSkipProperty = Super::ShouldSkipProperty(InProperty);
 	
-	const bool bShouldSerializeAllValues = !SelectionSet.IsSet();
-	if (!bShouldSkipProperty && !bShouldSerializeAllValues)
+	if (!bShouldSkipProperty && !ShouldSerializeAllProperties())
 	{
-		bShouldSkipProperty = !SelectionSet.GetValue()->ShouldSerializeProperty(GetSerializedPropertyChain(), InProperty) || !CastField<FObjectPropertyBase>(InProperty);
+		const bool bIsAllowed = SelectionSet.GetValue()->ShouldSerializeProperty(GetSerializedPropertyChain(), InProperty) || CastField<FObjectPropertyBase>(InProperty);  
+		bShouldSkipProperty = !bIsAllowed;
 	}
 	
 	return bShouldSkipProperty;
@@ -136,20 +218,25 @@ void FApplySnapshotDataArchiveV2::PopSerializedProperty(FProperty* InProperty, c
 #endif
 }
 
-FApplySnapshotDataArchiveV2::FApplySnapshotDataArchiveV2(FObjectSnapshotData& InObjectData, FWorldSnapshotData& InSharedData, UObject* InOriginalObject, const FPropertySelection& InSelectionSet)
+UObject* FApplySnapshotDataArchiveV2::ResolveObjectDependency(int32 ObjectIndex) const
+{
+	return GetSharedData().ResolveObjectDependencyForEditorWorld(ObjectIndex, SelectionMapForResolvingSubobjects);
+}
+
+FApplySnapshotDataArchiveV2::FApplySnapshotDataArchiveV2(FObjectSnapshotData& InObjectData, FWorldSnapshotData& InSharedData, UObject* InOriginalObject, const FPropertySelectionMap& InSelectionMapForResolvingSubobjects, TOptional<const FPropertySelection*> InSelectionSet)
         :
         Super(InObjectData, InSharedData, true),
-        SelectionSet(&InSelectionSet),
-		PropertiesLeftToSerialize(InSelectionSet),
-        OriginalObject(InOriginalObject)
+        OriginalObject(InOriginalObject),
+		SelectionMapForResolvingSubobjects(InSelectionMapForResolvingSubobjects),
+		SelectionSet(InSelectionSet)
 {
-	bShouldLoadObjectDependenciesForTempWorld = false;
+	if (SelectionSet)
+	{
+		PropertiesLeftToSerialize = *SelectionSet.GetValue();
+	}
 }
 
-FApplySnapshotDataArchiveV2::FApplySnapshotDataArchiveV2(FObjectSnapshotData& InObjectData, FWorldSnapshotData& InSharedData, UObject* InOriginalObject)
-	:
-	Super(InObjectData, InSharedData, true),
-	OriginalObject(InOriginalObject)
+bool FApplySnapshotDataArchiveV2::ShouldSerializeAllProperties() const
 {
-	bShouldLoadObjectDependenciesForTempWorld = false;
+	return !SelectionSet.IsSet();
 }
diff --git a/Engine/Plugins/VirtualProduction/LevelSnapshots/Source/LevelSnapshots/Private/Archive/ApplySnapshotDataArchiveV2.h b/Engine/Plugins/VirtualProduction/LevelSnapshots/Source/LevelSnapshots/Private/Archive/ApplySnapshotDataArchiveV2.h
--- a/Engine/Plugins/VirtualProduction/LevelSnapshots/Source/LevelSnapshots/Private/Archive/ApplySnapshotDataArchiveV2.h
+++ b/Engine/Plugins/VirtualProduction/LevelSnapshots/Source/LevelSnapshots/Private/Archive/ApplySnapshotDataArchiveV2.h
@@ -4,18 +4,17 @@
 
 #include "CoreMinimal.h"
 #include "PropertySelection.h"
+#include "PropertySelectionMap.h"
 #include "SnapshotArchive.h"
 
-class ULevelSnapshotSelectionSet;
-
 /* For writing data into an object */
 class FApplySnapshotDataArchiveV2 : public FSnapshotArchive
 {
 	using Super = FSnapshotArchive;
 public:
 	
-	static void ApplyToExistingWorldObject(FObjectSnapshotData& InObjectData, FWorldSnapshotData& InSharedData, UObject* InOriginalObject, UObject* InDeserializedVersion, const FPropertySelection& InSelectionSet);
-	static void ApplyToRecreatedWorldObject(FObjectSnapshotData& InObjectData, FWorldSnapshotData& InSharedData, UObject* InOriginalObject, UObject* InDeserializedVersion);
+	static void ApplyToExistingEditorWorldObject(FObjectSnapshotData& InObjectData, FWorldSnapshotData& InSharedData, UObject* InOriginalObject, UObject* InDeserializedVersion, const FPropertySelectionMap& InSelectionMapForResolvingSubobjects, const FPropertySelection& InSelectionSet);
+	static void ApplyToRecreatedEditorWorldObject(FObjectSnapshotData& InObjectData, FWorldSnapshotData& InSharedData, UObject* InOriginalObject, UObject* InDeserializedVersion, const FPropertySelectionMap& InSelectionMapForResolvingSubobjects);
 	
 	//~ Begin FTakeSnapshotArchiveV2 Interface
 	virtual bool ShouldSkipProperty(const FProperty* InProperty) const override;
@@ -23,10 +22,25 @@ public:
 	virtual void PopSerializedProperty(FProperty* InProperty, const bool bIsEditorOnlyProperty) override;
 	//~ End FTakeSnapshotArchiveV2 Interface
 
+protected:
+
+	//~ Begin FSnapshotArchive Interface
+	virtual UObject* ResolveObjectDependency(int32 ObjectIndex) const override;
+	//~ End FSnapshotArchive Interface
+
 private:
 
-	FApplySnapshotDataArchiveV2(FObjectSnapshotData& InObjectData, FWorldSnapshotData& InSharedData, UObject* InOriginalObject, const FPropertySelection& InSelectionSet);
-	FApplySnapshotDataArchiveV2(FObjectSnapshotData& InObjectData, FWorldSnapshotData& InSharedData, UObject* InOriginalObject);
+	FApplySnapshotDataArchiveV2(FObjectSnapshotData& InObjectData, FWorldSnapshotData& InSharedData, UObject* InOriginalObject, const FPropertySelectionMap& InSelectionMapForResolvingSubobjects, TOptional<const FPropertySelection*> InSelectionSet);
+
+	bool ShouldSerializeAllProperties() const;
+	
+	/* Object we are serializing into.
+	* Needed for Pre and Post Edit change
+	*/
+	UObject* OriginalObject;
+
+	/* Needed so subobjects can be fully reconstructed */
+	const FPropertySelectionMap& SelectionMapForResolvingSubobjects; 
 	
 	/* Immutable list of properties we are supposed to serialise */
 	TOptional<const FPropertySelection*> SelectionSet;
@@ -35,9 +49,5 @@ private:
 	 */
 	mutable FPropertySelection PropertiesLeftToSerialize;
 	
-	/* Object we are serializing into.
-	 * Needed for Pre and Post Edit change
-	 */
-	UObject* OriginalObject;
 	
 };
diff --git a/Engine/Plugins/VirtualProduction/LevelSnapshots/Source/LevelSnapshots/Private/Archive/ClassDefaults/ApplyClassDefaulDataArchive.cpp b/Engine/Plugins/VirtualProduction/LevelSnapshots/Source/LevelSnapshots/Private/Archive/ClassDefaults/ApplyClassDefaulDataArchive.cpp
new file mode 100644
--- /dev/null
+++ b/Engine/Plugins/VirtualProduction/LevelSnapshots/Source/LevelSnapshots/Private/Archive/ClassDefaults/ApplyClassDefaulDataArchive.cpp
@@ -0,0 +1,26 @@
+// Copyright Epic Games, Inc. All Rights Reserved.
+
+#include "Archive/ClassDefaults/ApplyClassDefaulDataArchive.h"
+
+#include "ObjectSnapshotData.h"
+#include "WorldSnapshotData.h"
+
+void FApplyClassDefaulDataArchive::SerializeClassDefaultObject(FObjectSnapshotData& InObjectData, FWorldSnapshotData& InSharedData, UObject* InCDO)
+{
+	FApplyClassDefaulDataArchive Archive(InObjectData, InSharedData, InCDO, ESerialisationMode::RestoringCDO);
+	InCDO->Serialize(Archive);
+}
+
+void FApplyClassDefaulDataArchive::RestoreChangedClassDefaults(FObjectSnapshotData& InObjectData, FWorldSnapshotData& InSharedData, UObject* InObjectToRestore)
+{
+	FApplyClassDefaulDataArchive Archive(InObjectData, InSharedData, InObjectToRestore, ESerialisationMode::RestoringChangedDefaults);
+	InObjectToRestore->Serialize(Archive);
+}
+
+FApplyClassDefaulDataArchive::FApplyClassDefaulDataArchive(FObjectSnapshotData& InObjectData, FWorldSnapshotData& InSharedData, UObject* InObjectToRestore, ESerialisationMode InSerialisationMode)
+	:
+	Super(InObjectData, InSharedData, true, InObjectToRestore)
+{
+	// Only overwrite transient properties when actually serializing into a snapshot CDO
+	ArSerializingDefaults = InSerialisationMode == ESerialisationMode::RestoringCDO;
+}
diff --git a/Engine/Plugins/VirtualProduction/LevelSnapshots/Source/LevelSnapshots/Private/Archive/ClassDefaults/ApplyClassDefaulDataArchive.h b/Engine/Plugins/VirtualProduction/LevelSnapshots/Source/LevelSnapshots/Private/Archive/ClassDefaults/ApplyClassDefaulDataArchive.h
new file mode 100644
--- /dev/null
+++ b/Engine/Plugins/VirtualProduction/LevelSnapshots/Source/LevelSnapshots/Private/Archive/ClassDefaults/ApplyClassDefaulDataArchive.h
@@ -0,0 +1,32 @@
+// Copyright Epic Games, Inc. All Rights Reserved.
+
+#pragma once
+
+#include "CoreMinimal.h"
+#include "Archive/ClassDefaults/BaseClassDefaultArchive.h"
+
+struct FObjectSnapshotData;
+struct FWorldSnapshotData;
+class UObject;
+
+/* Used when we're taking a snapshot of the world. */
+class FApplyClassDefaulDataArchive : public FBaseClassDefaultArchive
+{
+	using Super = FBaseClassDefaultArchive;
+public:
+
+	/* Archive will be used for reconstructing a CDO */
+	static void SerializeClassDefaultObject(FObjectSnapshotData& InObjectData, FWorldSnapshotData& InSharedData, UObject* InCDO);
+	/* Archive will be used to to apply any non-object properties before the object receives its saved snapshot values. This is to handle when a CDO has changed its default values since a snapshot was taken. */
+	static void RestoreChangedClassDefaults(FObjectSnapshotData& InObjectData, FWorldSnapshotData& InSharedData, UObject* ObjectToRestore);
+
+protected:
+	
+	enum class ESerialisationMode
+	{
+		RestoringCDO,
+		RestoringChangedDefaults
+	};
+
+	FApplyClassDefaulDataArchive(FObjectSnapshotData& InObjectData, FWorldSnapshotData& InSharedData, UObject* InObjectToRestore, ESerialisationMode InSerialisationMode);
+};
diff --git a/Engine/Plugins/VirtualProduction/LevelSnapshots/Source/LevelSnapshots/Private/Archive/ClassDefaults/BaseClassDefaultArchive.cpp b/Engine/Plugins/VirtualProduction/LevelSnapshots/Source/LevelSnapshots/Private/Archive/ClassDefaults/BaseClassDefaultArchive.cpp
new file mode 100644
--- /dev/null
+++ b/Engine/Plugins/VirtualProduction/LevelSnapshots/Source/LevelSnapshots/Private/Archive/ClassDefaults/BaseClassDefaultArchive.cpp
@@ -0,0 +1,67 @@
+﻿// Copyright Epic Games, Inc. All Rights Reserved.
+
+#include "BaseClassDefaultArchive.h"
+
+#include "WorldSnapshotData.h"
+
+bool FBaseClassDefaultArchive::ShouldSkipProperty(const FProperty* InProperty) const
+{
+	const bool bSuperWantsToSkip = Super::ShouldSkipProperty(InProperty);
+	
+	// Avoid any object references owned by a CDO. References are automatically replaced by object instance creation flow.
+	return bSuperWantsToSkip || IsPropertyReferenceToSubobjectOrClassDefaults(InProperty);
+}
+
+UObject* FBaseClassDefaultArchive::ResolveObjectDependency(int32 ObjectIndex) const
+{
+	return GetSharedData().ResolveObjectDependencyForClassDefaultObject(ObjectIndex);
+}
+
+FBaseClassDefaultArchive::FBaseClassDefaultArchive(FObjectSnapshotData& InObjectData, FWorldSnapshotData& InSharedData, bool bIsLoading, UObject* InObjectToRestore)
+	:
+	Super(InObjectData, InSharedData, bIsLoading),
+	SerializedObject(InObjectToRestore)
+{
+	// Description of CPF_Transient: "Property is transient: shouldn't be saved or loaded, except for Blueprint CDOs."
+	ExcludedPropertyFlags = CPF_BlueprintAssignable | CPF_Deprecated
+		// Do not save any instanced references when serialising CDOs
+		| CPF_ContainsInstancedReference | CPF_InstancedReference | CPF_PersistentInstance;
+	
+	// Otherwise we are not allowed to serialize transient properties
+	ArSerializingDefaults = true;
+}
+
+bool FBaseClassDefaultArchive::IsPropertyReferenceToSubobjectOrClassDefaults(const FProperty* InProperty) const
+{
+	const FObjectPropertyBase* ObjectProperty = CastField<FObjectPropertyBase>(InProperty);
+	if (!ObjectProperty)
+	{
+		return false;
+	}
+
+	const bool bIsMarkedAsSubobject = InProperty->HasAnyPropertyFlags(CPF_InstancedReference | CPF_ContainsInstancedReference | CPF_PersistentInstance);
+	const bool bIsActorOrComponentPtr = ObjectProperty->PropertyClass->IsChildOf(AActor::StaticClass()) || ObjectProperty->PropertyClass->IsChildOf(UActorComponent::StaticClass());
+	if (bIsMarkedAsSubobject || bIsActorOrComponentPtr)
+	{
+		return true;
+	}
+
+	const FArchiveSerializedPropertyChain* PropertyChain = GetSerializedPropertyChain();
+	const void* ContainerPtr = SerializedObject;
+	for (int32 i = 0; PropertyChain && i < PropertyChain->GetNumProperties(); ++i)
+	{
+		ContainerPtr = PropertyChain->GetPropertyFromRoot(i)->ContainerPtrToValuePtr<void>(ContainerPtr);
+	}
+
+	const void* LeafValuePtr = InProperty->ContainerPtrToValuePtr<void>(ContainerPtr);
+	if (const UObject* ContainedPtr = ObjectProperty->GetObjectPropertyValue(LeafValuePtr))
+	{
+		const bool bIsClassDefault = ContainedPtr->HasAnyFlags(RF_ClassDefaultObject | RF_ArchetypeObject);
+		const bool bIsPointingToDefaultSubobject = ContainedPtr->HasAnyFlags(RF_DefaultSubObject);
+		const bool bIsPointingToSelf = ContainedPtr == SerializedObject;
+		const bool bIsPointingToSubobject = ContainedPtr->IsIn(SerializedObject);
+		return bIsClassDefault || bIsPointingToDefaultSubobject || bIsPointingToSelf || bIsPointingToSubobject;
+	}
+	
+	return false;
+}
\ No newline at end of file
diff --git a/Engine/Plugins/VirtualProduction/LevelSnapshots/Source/LevelSnapshots/Private/Archive/ClassDefaults/BaseClassDefaultArchive.h b/Engine/Plugins/VirtualProduction/LevelSnapshots/Source/LevelSnapshots/Private/Archive/ClassDefaults/BaseClassDefaultArchive.h
new file mode 100644
--- /dev/null
+++ b/Engine/Plugins/VirtualProduction/LevelSnapshots/Source/LevelSnapshots/Private/Archive/ClassDefaults/BaseClassDefaultArchive.h
@@ -0,0 +1,37 @@
+﻿// Copyright Epic Games, Inc. All Rights Reserved.
+
+#pragma once
+
+#include "CoreMinimal.h"
+#include "Archive/SnapshotArchive.h"
+
+class UObject;
+struct FObjectSnapshotData;
+struct FWorldSnapshotData;
+
+/* Shared logic for serializing class defaults. */
+class FBaseClassDefaultArchive : public FSnapshotArchive
+{
+	using Super = FSnapshotArchive;
+public:
+
+	//~ Begin FSnapshotArchive Interface
+	virtual bool ShouldSkipProperty(const FProperty* InProperty) const override;
+	//~ End FSnapshotArchive Interface
+
+protected:
+	
+	/* The object we're serializing into. Can be a CDO or 'normal' object. Used to detect subobjects references. */
+	UObject* SerializedObject;
+
+	
+	//~ Begin FSnapshotArchive Interface
+	virtual UObject* ResolveObjectDependency(int32 ObjectIndex) const override;
+	//~ End FSnapshotArchive Interface
+	
+	FBaseClassDefaultArchive(FObjectSnapshotData& InObjectData, FWorldSnapshotData& InSharedData, bool bIsLoading, UObject* InObjectToRestore);
+
+private:
+	
+	bool IsPropertyReferenceToSubobjectOrClassDefaults(const FProperty* InProperty) const;
+};
diff --git a/Engine/Plugins/VirtualProduction/LevelSnapshots/Source/LevelSnapshots/Private/Archive/ClassDefaults/TakeClassDefaultObjectSnapshotArchive.cpp b/Engine/Plugins/VirtualProduction/LevelSnapshots/Source/LevelSnapshots/Private/Archive/ClassDefaults/TakeClassDefaultObjectSnapshotArchive.cpp
new file mode 100644
--- /dev/null
+++ b/Engine/Plugins/VirtualProduction/LevelSnapshots/Source/LevelSnapshots/Private/Archive/ClassDefaults/TakeClassDefaultObjectSnapshotArchive.cpp
@@ -0,0 +1,17 @@
+// Copyright Epic Games, Inc. All Rights Reserved.
+
+#include "Archive/ClassDefaults/TakeClassDefaultObjectSnapshotArchive.h"
+
+#include "ObjectSnapshotData.h"
+#include "WorldSnapshotData.h"
+
+void FTakeClassDefaultObjectSnapshotArchive::SaveClassDefaultObject(FObjectSnapshotData& InObjectData, FWorldSnapshotData& InSharedData, UObject* SerializedObject)
+{
+	FTakeClassDefaultObjectSnapshotArchive SaveClass(InObjectData, InSharedData, SerializedObject);
+	SerializedObject->Serialize(SaveClass);
+}
+
+FTakeClassDefaultObjectSnapshotArchive::FTakeClassDefaultObjectSnapshotArchive(FObjectSnapshotData& InObjectData, FWorldSnapshotData& InSharedData, UObject* InSerializedObject)
+	:
+	Super(InObjectData, InSharedData, false, InSerializedObject)
+{}
diff --git a/Engine/Plugins/VirtualProduction/LevelSnapshots/Source/LevelSnapshots/Private/Archive/ClassDefaults/TakeClassDefaultObjectSnapshotArchive.h b/Engine/Plugins/VirtualProduction/LevelSnapshots/Source/LevelSnapshots/Private/Archive/ClassDefaults/TakeClassDefaultObjectSnapshotArchive.h
new file mode 100644
--- /dev/null
+++ b/Engine/Plugins/VirtualProduction/LevelSnapshots/Source/LevelSnapshots/Private/Archive/ClassDefaults/TakeClassDefaultObjectSnapshotArchive.h
@@ -0,0 +1,23 @@
+// Copyright Epic Games, Inc. All Rights Reserved.
+
+#pragma once
+
+#include "CoreMinimal.h"
+#include "Archive/ClassDefaults/BaseClassDefaultArchive.h"
+
+class UObject;
+struct FObjectSnapshotData;
+struct FWorldSnapshotData;
+
+/* Used when we're taking a snapshot of the world. */
+class FTakeClassDefaultObjectSnapshotArchive final : public FBaseClassDefaultArchive
+{
+	using Super = FBaseClassDefaultArchive;
+public:
+
+	static void SaveClassDefaultObject(FObjectSnapshotData& InObjectData, FWorldSnapshotData& InSharedData, UObject* SerializedObject);
+
+private:
+	
+	FTakeClassDefaultObjectSnapshotArchive(FObjectSnapshotData& InObjectData, FWorldSnapshotData& InSharedData, UObject* InSerializedObject);
+};
diff --git a/Engine/Plugins/VirtualProduction/LevelSnapshots/Source/LevelSnapshots/Private/Archive/SnapshotArchive.cpp b/Engine/Plugins/VirtualProduction/LevelSnapshots/Source/LevelSnapshots/Private/Archive/SnapshotArchive.cpp
--- a/Engine/Plugins/VirtualProduction/LevelSnapshots/Source/LevelSnapshots/Private/Archive/SnapshotArchive.cpp
+++ b/Engine/Plugins/VirtualProduction/LevelSnapshots/Source/LevelSnapshots/Private/Archive/SnapshotArchive.cpp
@@ -5,11 +5,15 @@
 #include "ObjectSnapshotData.h"
 #include "WorldSnapshotData.h"
 
+#include "Internationalization/TextNamespaceUtil.h"
+#include "Internationalization/TextPackageNamespaceUtil.h"
 #include "UObject/ObjectMacros.h"
 
-FSnapshotArchive FSnapshotArchive::MakeArchiveForRestoring(FObjectSnapshotData& InObjectData, FWorldSnapshotData& InSharedData)
+void FSnapshotArchive::RestoreData(FObjectSnapshotData& InObjectData, FWorldSnapshotData& InSharedData, UObject* InObjectToRestore, UPackage* InLocalisationSnapshotPackage)
 {
-	return FSnapshotArchive(InObjectData, InSharedData, true);
+	FSnapshotArchive Archive(InObjectData, InSharedData, true);
+	Archive.SetLocalizationNamespace(TextNamespaceUtil::EnsurePackageNamespace(InLocalisationSnapshotPackage));
+	InObjectToRestore->Serialize(Archive);
 }
 
 FString FSnapshotArchive::GetArchiveName() const
@@ -36,7 +40,9 @@ void FSnapshotArchive::Seek(int64 InPos)
 bool FSnapshotArchive::ShouldSkipProperty(const FProperty* InProperty) const
 {
 	return InProperty->HasAnyPropertyFlags(ExcludedPropertyFlags)
-		// We do not support (instanced) subobjects at this moment
+		// CPF_InstancedReference and CPF_ContainsInstancedReference skips properties skip references to components.
+		// CPF_PersistentInstance skips things like UPROPERTY(Instanced)
+		// Note, this still allows subobjects, e.g. when construction script creates a new material instance.
 		|| InProperty->HasAnyPropertyFlags(CPF_InstancedReference | CPF_ContainsInstancedReference | CPF_PersistentInstance);
 }
 
@@ -84,13 +90,11 @@ FArchive& FSnapshotArchive::operator<<(UObject*& Value)
 			return *this;
 		}
 
-		TOptional<UObject*> Resolved = SharedData.ResolveObjectDependency(ReferencedIndex, bShouldLoadObjectDependenciesForTempWorld ? FWorldSnapshotData::EResolveType::ResolveForUseInTempWorld : FWorldSnapshotData::EResolveType::ResolveForUseInOriginalWorld);
-		Value = Resolved.Get(nullptr);
+		Value = ResolveObjectDependency(ReferencedIndex);
 	}
 	else
 	{
 		int32 ReferenceIndex = SharedData.AddObjectDependency(Value);
-		// TODO: Check whether subobject and allocate
 		*this << ReferenceIndex;
 	}
 	
@@ -128,6 +132,11 @@ void FSnapshotArchive::Serialize(void* Data, int64 Length)
 	}
 }
 
+UObject* FSnapshotArchive::ResolveObjectDependency(int32 ObjectIndex) const
+{
+	return SharedData.ResolveObjectDependencyForSnapshotWorld(ObjectIndex);
+}
+
 FSnapshotArchive::FSnapshotArchive(FObjectSnapshotData& InObjectData, FWorldSnapshotData& InSharedData, bool bIsLoading)
 	:
 	ExcludedPropertyFlags(CPF_BlueprintAssignable | CPF_Transient | CPF_Deprecated),
@@ -142,9 +151,11 @@ FSnapshotArchive::FSnapshotArchive(FObjectSnapshotData& InObjectData, FWorldSnap
 	if (bIsLoading)
 	{
 		Super::SetIsLoading(true);
+		Super::SetIsSaving(false);
 	}
 	else
 	{
+		Super::SetIsLoading(false);
 		Super::SetIsSaving(true);
 	}
 }
diff --git a/Engine/Plugins/VirtualProduction/LevelSnapshots/Source/LevelSnapshots/Private/Archive/SnapshotArchive.h b/Engine/Plugins/VirtualProduction/LevelSnapshots/Source/LevelSnapshots/Private/Archive/SnapshotArchive.h
--- a/Engine/Plugins/VirtualProduction/LevelSnapshots/Source/LevelSnapshots/Private/Archive/SnapshotArchive.h
+++ b/Engine/Plugins/VirtualProduction/LevelSnapshots/Source/LevelSnapshots/Private/Archive/SnapshotArchive.h
@@ -5,6 +5,8 @@
 #include "CoreMinimal.h"
 #include "Serialization/ArchiveUObject.h"
 
+class UObject;
+class UPackage;
 struct FObjectSnapshotData;
 struct FWorldSnapshotData;
 
@@ -15,7 +17,7 @@ class FSnapshotArchive : public FArchiveUObject
 	
 public:
 
-	static FSnapshotArchive MakeArchiveForRestoring(FObjectSnapshotData& InObjectData, FWorldSnapshotData& InSharedData);
+	static void RestoreData(FObjectSnapshotData& InObjectData, FWorldSnapshotData& InSharedData, UObject* InObjectToRestore, UPackage* InLocalisationSnapshotPackage);
 
 	//~ Begin FArchive Interface
 	virtual FString GetArchiveName() const override;
@@ -30,18 +32,14 @@ public:
 	//~ End FArchive Interface
 
 protected:
-
+	
+	/* Allocates and serializes an object dependency, or gets the object, if it already exists. */
+	virtual UObject* ResolveObjectDependency(int32 ObjectIndex) const;
+	
 	FWorldSnapshotData& GetSharedData() const
 	{
 		return SharedData;
 	}
-	
-	/* Only used when Loading.
-	 * True: we're loading into a temp world. False: we're loading into an editor world.
-	 * 
-	 * If true, object references within the world are translated to use the temp world package.
-	 */
-	bool bShouldLoadObjectDependenciesForTempWorld = true;
 
 	int32 ExcludedPropertyFlags;
 
diff --git a/Engine/Plugins/VirtualProduction/LevelSnapshots/Source/LevelSnapshots/Private/Archive/TakeClassDefaultObjectSnapshotArchive.cpp b/Engine/Plugins/VirtualProduction/LevelSnapshots/Source/LevelSnapshots/Private/Archive/TakeClassDefaultObjectSnapshotArchive.cpp
deleted file mode 100644
diff --git a/Engine/Plugins/VirtualProduction/LevelSnapshots/Source/LevelSnapshots/Private/Archive/TakeClassDefaultObjectSnapshotArchive.h b/Engine/Plugins/VirtualProduction/LevelSnapshots/Source/LevelSnapshots/Private/Archive/TakeClassDefaultObjectSnapshotArchive.h
deleted file mode 100644
diff --git a/Engine/Plugins/VirtualProduction/LevelSnapshots/Source/LevelSnapshots/Private/Data/ActorSnapshotData.cpp b/Engine/Plugins/VirtualProduction/LevelSnapshots/Source/LevelSnapshots/Private/Data/ActorSnapshotData.cpp
--- a/Engine/Plugins/VirtualProduction/LevelSnapshots/Source/LevelSnapshots/Private/Data/ActorSnapshotData.cpp
+++ b/Engine/Plugins/VirtualProduction/LevelSnapshots/Source/LevelSnapshots/Private/Data/ActorSnapshotData.cpp
@@ -11,6 +11,7 @@
 
 #include "Engine/World.h"
 #include "GameFramework/Actor.h"
+#include "UObject/Package.h"
 
 FActorSnapshotData FActorSnapshotData::SnapshotActor(AActor* OriginalActor, FWorldSnapshotData& WorldData)
 {
@@ -38,7 +39,7 @@ FActorSnapshotData FActorSnapshotData::SnapshotActor(AActor* OriginalActor, FWor
 	return Result;
 }
 
-void FActorSnapshotData::DeserializeIntoExistingWorldActor(UWorld* SnapshotWorld, AActor* OriginalActor, FWorldSnapshotData& WorldData, const FPropertySelectionMap& SelectedProperties)
+void FActorSnapshotData::DeserializeIntoExistingWorldActor(UWorld* SnapshotWorld, AActor* OriginalActor, FWorldSnapshotData& WorldData, UPackage* InLocalisationSnapshotPackage, const FPropertySelectionMap& SelectedProperties)
 {
 	auto DeserializeActor = [this, &WorldData, &SelectedProperties](AActor* OriginalActor, AActor* DeserializedActor)
 	{
@@ -48,7 +49,7 @@ void FActorSnapshotData::DeserializeIntoExistingWorldActor(UWorld* SnapshotWorld
 #if WITH_EDITOR
 			OriginalActor->Modify();
 #endif
-			FApplySnapshotDataArchiveV2::ApplyToExistingWorldObject(SerializedActorData, WorldData, OriginalActor, DeserializedActor, *ActorPropertySelection);
+			FApplySnapshotDataArchiveV2::ApplyToExistingEditorWorldObject(SerializedActorData, WorldData, OriginalActor, DeserializedActor, SelectedProperties, *ActorPropertySelection);
 		}
 	};
 	auto DeserializeComponent = [&SelectedProperties, &WorldData](FObjectSnapshotData& SerializedCompData, FComponentSnapshotData& CompData, UActorComponent* Original, UActorComponent* Deserialized)
@@ -59,31 +60,31 @@ void FActorSnapshotData::DeserializeIntoExistingWorldActor(UWorld* SnapshotWorld
 #if WITH_EDITOR
 			Original->Modify();
 #endif
-			CompData.DeserializeIntoExistingWorldActor(SerializedCompData, Original, Deserialized, WorldData, *ComponentSelectedProperties);
+			FApplySnapshotDataArchiveV2::ApplyToExistingEditorWorldObject(SerializedCompData, WorldData, Original, Deserialized, SelectedProperties, *ComponentSelectedProperties);
 		};
 	};
 
-	DeserializeIntoWorldActor(SnapshotWorld, OriginalActor, WorldData, DeserializeActor, DeserializeComponent);
+	DeserializeIntoWorldActor(SnapshotWorld, OriginalActor, WorldData, InLocalisationSnapshotPackage, DeserializeActor, DeserializeComponent);
 }
 
-void FActorSnapshotData::DeserializeIntoRecreatedWorldActor(UWorld* SnapshotWorld, AActor* OriginalActor, FWorldSnapshotData& WorldData)
+void FActorSnapshotData::DeserializeIntoRecreatedEditorWorldActor(UWorld* SnapshotWorld, AActor* OriginalActor, FWorldSnapshotData& WorldData, UPackage* InLocalisationSnapshotPackage, const FPropertySelectionMap& SelectedProperties)
 {
-	auto DeserializeActor = [this, &WorldData](AActor* OriginalActor, AActor* DeserializedActor)
+	auto DeserializeActor = [this, &WorldData, &SelectedProperties](AActor* OriginalActor, AActor* DeserializedActor)
 	{
 #if WITH_EDITOR
 		OriginalActor->Modify();
 #endif
-		FApplySnapshotDataArchiveV2::ApplyToRecreatedWorldObject(SerializedActorData, WorldData, OriginalActor, DeserializedActor);
+		FApplySnapshotDataArchiveV2::ApplyToRecreatedEditorWorldObject(SerializedActorData, WorldData, OriginalActor, DeserializedActor, SelectedProperties);
 	};
-	auto DeserializeComponent = [&WorldData](FObjectSnapshotData& SerializedCompData, FComponentSnapshotData& CompData, UActorComponent* Original, UActorComponent* Deserialized)
+	auto DeserializeComponent = [&WorldData, &SelectedProperties](FObjectSnapshotData& SerializedCompData, FComponentSnapshotData& CompData, UActorComponent* Original, UActorComponent* Deserialized)
 	{
 #if WITH_EDITOR
 		Original->Modify();
 #endif
-		CompData.DeserializeIntoRecreatedWorldActor(SerializedCompData, Original, Deserialized, WorldData);
+		FApplySnapshotDataArchiveV2::ApplyToRecreatedEditorWorldObject(SerializedCompData, WorldData, Original, Deserialized, SelectedProperties); 
 	};
 	
-	DeserializeIntoWorldActor(SnapshotWorld, OriginalActor, WorldData, DeserializeActor, DeserializeComponent);
+	DeserializeIntoWorldActor(SnapshotWorld, OriginalActor, WorldData, InLocalisationSnapshotPackage, DeserializeActor, DeserializeComponent);
 }
 
 TOptional<AActor*> FActorSnapshotData::GetPreallocatedIfValidButDoNotAllocate() const
@@ -126,7 +127,7 @@ TOptional<AActor*> FActorSnapshotData::GetPreallocated(UWorld* SnapshotWorld, FW
 	return CachedSnapshotActor.IsValid() ? TOptional<AActor*>(CachedSnapshotActor.Get()) : TOptional<AActor*>();
 }
 
-TOptional<AActor*> FActorSnapshotData::GetDeserialized(UWorld* SnapshotWorld, FWorldSnapshotData& WorldData)
+TOptional<AActor*> FActorSnapshotData::GetDeserialized(UWorld* SnapshotWorld, FWorldSnapshotData& WorldData, UPackage* InLocalisationSnapshotPackage)
 {
 	if (bReceivedSerialisation && CachedSnapshotActor.IsValid())
 	{
@@ -141,12 +142,11 @@ TOptional<AActor*> FActorSnapshotData::GetDeserialized(UWorld* SnapshotWorld, FW
 	bReceivedSerialisation = true;
 	
 	AActor* PreallocatedActor = Preallocated.GetValue();
-	FSnapshotArchive Serializer = FSnapshotArchive::MakeArchiveForRestoring(SerializedActorData, WorldData);
-	PreallocatedActor->Serialize(Serializer);
+	FSnapshotArchive::RestoreData(SerializedActorData, WorldData, PreallocatedActor, InLocalisationSnapshotPackage);
 
-	DeserializeComponents(PreallocatedActor, WorldData, [](FObjectSnapshotData& SerializedCompData, FComponentSnapshotData& CompData, UActorComponent* Comp, FWorldSnapshotData& SharedData)
+	DeserializeComponents(PreallocatedActor, WorldData, [InLocalisationSnapshotPackage](FObjectSnapshotData& SerializedCompData, FComponentSnapshotData& CompData, UActorComponent* Comp, FWorldSnapshotData& SharedData)
 	{
-		CompData.DeserializeIntoTransient(SerializedCompData, Comp, SharedData);
+		CompData.DeserializeIntoTransient(SerializedCompData, Comp, SharedData, InLocalisationSnapshotPackage);
 	});
 
 	PreallocatedActor->UpdateComponentTransforms();
@@ -158,9 +158,9 @@ FSoftClassPath FActorSnapshotData::GetActorClass() const
 	return ActorClass;
 }
 
-void FActorSnapshotData::DeserializeIntoWorldActor(UWorld* SnapshotWorld, AActor* OriginalActor, FWorldSnapshotData& WorldData, FSerializeActor SerializeActor, FSerializeComponent SerializeComponent)
+void FActorSnapshotData::DeserializeIntoWorldActor(UWorld* SnapshotWorld, AActor* OriginalActor, FWorldSnapshotData& WorldData, UPackage* InLocalisationSnapshotPackage, FSerializeActor SerializeActor, FSerializeComponent SerializeComponent)
 {
-	const TOptional<AActor*> Deserialized = GetDeserialized(SnapshotWorld, WorldData);
+	const TOptional<AActor*> Deserialized = GetDeserialized(SnapshotWorld, WorldData, InLocalisationSnapshotPackage);
 	if (!Deserialized)
 	{
 		UE_LOG(LogLevelSnapshots, Warning, TEXT("Failed to serialize into actor %s. Skipping..."), *OriginalActor->GetName());
diff --git a/Engine/Plugins/VirtualProduction/LevelSnapshots/Source/LevelSnapshots/Private/Data/ComponentSnapshotData.cpp b/Engine/Plugins/VirtualProduction/LevelSnapshots/Source/LevelSnapshots/Private/Data/ComponentSnapshotData.cpp
--- a/Engine/Plugins/VirtualProduction/LevelSnapshots/Source/LevelSnapshots/Private/Data/ComponentSnapshotData.cpp
+++ b/Engine/Plugins/VirtualProduction/LevelSnapshots/Source/LevelSnapshots/Private/Data/ComponentSnapshotData.cpp
@@ -2,13 +2,12 @@
 
 #include "ComponentSnapshotData.h"
 
-#include "ApplySnapshotDataArchiveV2.h"
 #include "LevelSnapshotsLog.h"
-#include "PropertySelection.h"
 #include "TakeWorldObjectSnapshotArchive.h"
 #include "WorldSnapshotData.h"
 
 #include "Components/ActorComponent.h"
+#include "UObject/Package.h"
 
 TOptional<FComponentSnapshotData> FComponentSnapshotData::SnapshotComponent(UActorComponent* OriginalComponent, FWorldSnapshotData& WorldData)
 {
@@ -32,7 +31,7 @@ bool FComponentSnapshotData::IsRestoreSupportedForSavedComponent()
 	return CreationMethod != EComponentCreationMethod::Instance && CreationMethod != EComponentCreationMethod::UserConstructionScript;
 }
 
-void FComponentSnapshotData::DeserializeIntoTransient(FObjectSnapshotData& SerializedComponentData, UActorComponent* ComponentToDeserializeInto, FWorldSnapshotData& WorldData)
+void FComponentSnapshotData::DeserializeIntoTransient(FObjectSnapshotData& SerializedComponentData, UActorComponent* ComponentToDeserializeInto, FWorldSnapshotData& WorldData, UPackage* InLocalisationSnapshotPackage)
 {
 	// TODO: Handle CreationMethod == Instance and UserConstructionScript here by re-creating the component if it does not exist
 	if (ComponentToDeserializeInto->CreationMethod == EComponentCreationMethod::Instance)
@@ -59,16 +58,5 @@ void FComponentSnapshotData::DeserializeIntoTransient(FObjectSnapshotData& Seria
 		//   - We apply the CDO and afterwards we override it with the serialized data. Good.
 	WorldData.SerializeClassDefaultsInto(ComponentToDeserializeInto);
 	
-	FSnapshotArchive Serializer = FSnapshotArchive::MakeArchiveForRestoring(SerializedComponentData, WorldData);
-	ComponentToDeserializeInto->Serialize(Serializer);
-}
-
-void FComponentSnapshotData::DeserializeIntoExistingWorldActor(FObjectSnapshotData& SerializedComponentData, UActorComponent* OriginalComponentToDeserializeInto, UActorComponent* DeserializedComponentCounterpart, FWorldSnapshotData& WorldData, const FPropertySelection& PropertySelection)
-{
-	FApplySnapshotDataArchiveV2::ApplyToExistingWorldObject(SerializedComponentData, WorldData, OriginalComponentToDeserializeInto, DeserializedComponentCounterpart, PropertySelection);
-}
-
-void FComponentSnapshotData::DeserializeIntoRecreatedWorldActor(FObjectSnapshotData& SerializedComponentData, UActorComponent* OriginalComponentToDeserializeInto, UActorComponent* DeserializedComponentCounterpart, FWorldSnapshotData& WorldData)
-{
-	FApplySnapshotDataArchiveV2::ApplyToRecreatedWorldObject(SerializedComponentData, WorldData, OriginalComponentToDeserializeInto, DeserializedComponentCounterpart);
+	FSnapshotArchive::RestoreData(SerializedComponentData, WorldData, ComponentToDeserializeInto, InLocalisationSnapshotPackage);
 }
diff --git a/Engine/Plugins/VirtualProduction/LevelSnapshots/Source/LevelSnapshots/Private/Data/LevelSnapshot.cpp b/Engine/Plugins/VirtualProduction/LevelSnapshots/Source/LevelSnapshots/Private/Data/LevelSnapshot.cpp
--- a/Engine/Plugins/VirtualProduction/LevelSnapshots/Source/LevelSnapshots/Private/Data/LevelSnapshot.cpp
+++ b/Engine/Plugins/VirtualProduction/LevelSnapshots/Source/LevelSnapshots/Private/Data/LevelSnapshot.cpp
@@ -14,6 +14,8 @@
 #include "GameFramework/Actor.h"
 #include "SnapshotRestorability.h"
 #include "UObject/Package.h"
+#include "UObject/TextProperty.h"
+#include "UObject/UnrealType.h"
 
 #if WITH_EDITOR
 #include "ScopedTransaction.h"
@@ -24,36 +26,17 @@ namespace
 	/* If this function return false, the objects are not equivalent. If true, ignore the object references. */
 	bool ShouldConsiderObjectsEquivalent(const FWorldSnapshotData& SnapshotData, const FObjectPropertyBase* ObjectProperty, void* SnapshotValuePtr, void* WorldValuePtr, AActor* SnapshotActor, AActor* WorldActor)
 	{
-		UObject* SnapshotObject = ObjectProperty->GetObjectPropertyValue(SnapshotValuePtr);
-		UObject* WorldObject = ObjectProperty->GetObjectPropertyValue(WorldValuePtr);
-		if (SnapshotObject == nullptr && WorldObject == nullptr)
-		{
-			return true;
-		}
-
-		UObject* PossiblySubobject = WorldObject ? WorldObject : SnapshotObject;
-		AActor* PossiblyOwner = WorldObject ?  WorldActor : SnapshotActor;
-		// Handle subobjects created within actor
-		const bool bIsSubobject = PossiblySubobject->IsIn(PossiblyOwner);
-		if (bIsSubobject)
-		{
-			return true;
-		}
-		// Handle temporary 'subobjects'
-		const bool bIsTempTransientObject = PossiblySubobject->HasAnyFlags(RF_Transient) || PossiblySubobject->GetPackage()->HasAnyFlags(RF_Transient);
-		if (bIsTempTransientObject)
-		{
-			return true;
-		}
-
-		// Handle internal reference to other objects within the same world
-		if (SnapshotData.AreReferencesEquivalent(SnapshotObject, WorldObject))
+		// Native identity check handles:
+			// - external references, e.g. UMaterial in content browser
+			// - soft object paths: if SnapshotValuePtr is a TSoftObjectPtr<AActor> property, it points to a world actor instead of to an equivalent snapshot actor
+		if (ObjectProperty->Identical(SnapshotValuePtr, WorldValuePtr, 0))
 		{
 			return true;
 		}
-
-		// Anything this far should be an asset reference, e.g. to a data asset or material.
-		return ObjectProperty->Identical(SnapshotValuePtr, WorldValuePtr, 0);
+		
+		UObject* SnapshotObject = ObjectProperty->GetObjectPropertyValue(SnapshotValuePtr);
+		UObject* WorldObject = ObjectProperty->GetObjectPropertyValue(WorldValuePtr);
+		return SnapshotData.AreReferencesEquivalent(SnapshotObject, WorldObject, SnapshotActor, WorldActor);
 	}
 	
 	void EnqueueMatchingComponents(TInlineComponentArray<TPair<UObject*, UObject*>>& SnapshotOriginalPairsToProcess, AActor* SnapshotActor, AActor* WorldActor)
@@ -111,7 +94,7 @@ void ULevelSnapshot::ApplySnapshotToWorld(UWorld* TargetWorld, const FPropertySe
 	}
 	else
 	{
-		SerializedData.ApplyToWorld(TargetWorld, SelectionSet);
+		SerializedData.ApplyToWorld(TargetWorld, GetPackage(), SelectionSet);
 	}
 }
 
@@ -259,6 +242,13 @@ bool ULevelSnapshot::AreSnapshotAndOriginalPropertiesEquivalent(const FProperty*
 		{
 			// TODO: Use custom function. Need to do something similar to UE4SetProperty_Private::IsPermutation
 		}
+
+		if (const FTextProperty* TextProperty = CastField<FTextProperty>(LeafProperty))
+		{
+			const FText& SnapshotText = TextProperty->GetPropertyValue_InContainer(SnapshotContainer);
+			const FText& WorldText = TextProperty->GetPropertyValue_InContainer(WorldContainer);
+			return SnapshotText.IdenticalTo(WorldText, ETextIdenticalModeFlags::None) || SnapshotText.ToString().Equals(WorldText.ToString());
+		}
 		
 		// Use normal property comparison for all other properties
 		if (!LeafProperty->Identical_InContainer(SnapshotContainer, WorldContainer, i, PPF_DeepComparison | PPF_DeepCompareDSOsOnly))
@@ -283,7 +273,7 @@ TOptional<AActor*> ULevelSnapshot::GetDeserializedActor(const FSoftObjectPath& O
 	}
 	else
 	{
-		return SerializedData.GetDeserializedActor(OriginalActorPath);
+		return SerializedData.GetDeserializedActor(OriginalActorPath, GetPackage());
 	}
 }
 
@@ -389,8 +379,10 @@ void ULevelSnapshot::EnsureWorldInitialised()
 	if (SnapshotContainerWorld == nullptr)
 	{
 		SnapshotContainerWorld = NewObject<UWorld>(GetTransientPackage(), NAME_None);
-		SnapshotContainerWorld->WorldType = EWorldType::EditorPreview; 
-		
+		SnapshotContainerWorld->WorldType = EWorldType::EditorPreview;
+
+		// Note: Do NOT create a FWorldContext for this world.
+		// If you do, the render thread will send render commands every tick (and crash cuz we do not init the scene below).
 		SnapshotContainerWorld->InitializeNewWorld(UWorld::InitializationValues()
 			.InitializeScenes(false)		// This is memory only world: no rendering
             .AllowAudioPlayback(false)
diff --git a/Engine/Plugins/VirtualProduction/LevelSnapshots/Source/LevelSnapshots/Private/Data/PropertySelection.cpp b/Engine/Plugins/VirtualProduction/LevelSnapshots/Source/LevelSnapshots/Private/Data/PropertySelection.cpp
--- a/Engine/Plugins/VirtualProduction/LevelSnapshots/Source/LevelSnapshots/Private/Data/PropertySelection.cpp
+++ b/Engine/Plugins/VirtualProduction/LevelSnapshots/Source/LevelSnapshots/Private/Data/PropertySelection.cpp
@@ -80,8 +80,8 @@ bool FPropertySelection::ShouldSerializeProperty(const FArchiveSerializedPropert
 		//		Serialize(FArchive& Ar) { Ar << SomeVar;  }
 		// }
 		// LeafProperty will be FVector::X, which we need to allow.
-		const FStructProperty* StructProperty = CastField<FStructProperty>(ParentProperty);
-		if (StructProperty->Struct->UseNativeSerialization())
+		const FStructProperty* StructProperty = CastField<FStructProperty>(ParentProperty); 
+		if (StructProperty && StructProperty->Struct->UseNativeSerialization())
 		{
 			return true;
 		}
diff --git a/Engine/Plugins/VirtualProduction/LevelSnapshots/Source/LevelSnapshots/Private/Data/WorldSnapshotData.cpp b/Engine/Plugins/VirtualProduction/LevelSnapshots/Source/LevelSnapshots/Private/Data/WorldSnapshotData.cpp
--- a/Engine/Plugins/VirtualProduction/LevelSnapshots/Source/LevelSnapshots/Private/Data/WorldSnapshotData.cpp
+++ b/Engine/Plugins/VirtualProduction/LevelSnapshots/Source/LevelSnapshots/Private/Data/WorldSnapshotData.cpp
@@ -2,14 +2,18 @@
 
 #include "Data/WorldSnapshotData.h"
 
-#include "Archive/TakeClassDefaultObjectSnapshotArchive.h"
 #include "Archive/TakeWorldObjectSnapshotArchive.h"
+#include "Archive/ClassDefaults/ApplyClassDefaulDataArchive.h"
+#include "Archive/ClassDefaults/TakeClassDefaultObjectSnapshotArchive.h"
 #include "LevelSnapshotsLog.h"
 #include "LevelSnapshotsStats.h"
 #include "PropertySelectionMap.h"
 #include "Restorability/SnapshotRestorability.h"
 
+#include "Components/ActorComponent.h"
 #include "EngineUtils.h"
+#include "GameFramework/Actor.h"
+#include "UObject/Package.h"
 #if WITH_EDITOR
 #include "Editor/UnrealEdEngine.h"
 #include "Internationalization/Internationalization.h"
@@ -21,10 +25,52 @@
 
 namespace
 {
+	/* If the path contains an actor, returns a new path to that actor. */
+	TOptional<FSoftObjectPath> ExtractActorFromPath(const FSoftObjectPath& OriginalObjectPath, bool& bIsPathToActorSubobject)
+	{
+		const static FString PersistentLevelString("PersistentLevel.");
+		const int32 PersistentLevelStringLength = PersistentLevelString.Len();
+		const FString& SubPathString = OriginalObjectPath.GetSubPathString();
+		const int32 IndexOfPersistentLevelInfo = SubPathString.Find(PersistentLevelString, ESearchCase::CaseSensitive);
+		if (IndexOfPersistentLevelInfo == INDEX_NONE)
+		{
+			return {};
+		}
+
+		const int32 DotAfterActorNameIndex = SubPathString.Find(TEXT("."), ESearchCase::IgnoreCase, ESearchDir::FromStart, IndexOfPersistentLevelInfo + PersistentLevelStringLength);
+		const bool bPathIsToActor = DotAfterActorNameIndex == INDEX_NONE;
+		
+		bIsPathToActorSubobject = !bPathIsToActor;
+		return bPathIsToActor ?
+			// Example /Game/MapName.MapName:PersistentLevel.StaticMeshActor_42
+			OriginalObjectPath
+			:
+			// Converts /Game/MapName.MapName:PersistentLevel.StaticMeshActor_42.StaticMeshComponent to /Game/MapName.MapName:PersistentLevel.StaticMeshActor_42
+			FSoftObjectPath(OriginalObjectPath.GetAssetPathName(), SubPathString.Left(DotAfterActorNameIndex));
+	}
+	
 	/* If Path contains a path to an actor, returns that actor.
 	 * Example: /Game/MapName.MapName:PersistentLevel.StaticMeshActor_42.StaticMeshComponent returns StaticMeshActor_42's data
 	 */
 	TOptional<FActorSnapshotData*> FindSavedActorDataUsingObjectPath(TMap<FSoftObjectPath, FActorSnapshotData>& ActorData, const FSoftObjectPath& OriginalObjectPath, bool& bIsPathToActorSubobject)
+	{
+		const TOptional<FSoftObjectPath> PathToActor = ExtractActorFromPath(OriginalObjectPath, bIsPathToActorSubobject);
+		if (!PathToActor.IsSet())
+		{
+			return {};	
+		}
+		
+		FActorSnapshotData* Result = ActorData.Find(*PathToActor);
+		UE_CLOG(Result == nullptr, LogLevelSnapshots, Warning, TEXT("Path %s looks like an actor path but no data was saved for it. Maybe it was a reference to an auto-generated actor, e.g. a brush or volume present in all worlds by default?"), *OriginalObjectPath.ToString());
+		return Result ? Result : TOptional<FActorSnapshotData*>();  
+	}
+
+
+	/* Takes an existing path to an actor's subobjects and replaces the actor bit with the path to another actor.
+	 *
+	 * E.g. /Game/MapName.MapName:PersistentLevel.StaticMeshActor_42.StaticMeshComponent could become /Game/MapName.MapName:PersistentLevel.SomeOtherActor.StaticMeshComponent
+	 */
+	FSoftObjectPath SetActorInPath(AActor* NewActor, const FSoftObjectPath& OriginalObjectPath)
 	{
 		const static FString PersistentLevelString("PersistentLevel.");
 		const int32 PersistentLevelStringLength = PersistentLevelString.Len();
@@ -36,27 +82,82 @@ namespace
 		}
 
 		const int32 DotAfterActorNameIndex = SubPathString.Find(TEXT("."), ESearchCase::IgnoreCase, ESearchDir::FromStart, IndexOfPersistentLevelInfo + PersistentLevelStringLength);
-		const bool bPathIsToActor = DotAfterActorNameIndex == INDEX_NONE;
-		if (bPathIsToActor)
+
+		const FSoftObjectPath PathToNewActor(NewActor);
+		// PersistentLevel.StaticMeshActor_42.StaticMeshComponent becomes .StaticMeshComponent
+		const FString PathAfterOriginalActor = SubPathString.Right(SubPathString.Len() - DotAfterActorNameIndex + 1); // + 1 to include the dot after the actor
+		return FSoftObjectPath(PathToNewActor.GetAssetPathName(), PathToNewActor.GetSubPathString() + PathAfterOriginalActor);
+	}
+
+	enum EEquivalenceResult
+	{
+		Equivalent,
+		NotEquivalent,
+		NotComparable
+	};
+	
+	EEquivalenceResult AreActorsEquivalent(UObject* SnapshotPropertyValue, UObject* OriginalPropertyValue, const TMap<FSoftObjectPath, FActorSnapshotData>& ActorData)
+	{
+		// Compare actors
+		if (AActor* OriginalActorReference = Cast<AActor>(OriginalPropertyValue))
 		{
-			// Example /Game/MapName.MapName:PersistentLevel.StaticMeshActor_42
-			bIsPathToActorSubobject = false;
-			
-			FActorSnapshotData* Result = ActorData.Find(OriginalObjectPath);
-			UE_CLOG(Result == nullptr, LogLevelSnapshots, Warning, TEXT("Path %s looks like an actor path but no data was saved for it. Maybe it was a reference to an auto-generated actor, e.g. a brush or volume present in all worlds by default?"), *OriginalObjectPath.ToString());
-			return Result ? Result : TOptional<FActorSnapshotData*>();
+			const FActorSnapshotData* SavedData = ActorData.Find(OriginalActorReference);
+			if (SavedData == nullptr)
+			{
+				return NotEquivalent;
+			}
+
+			// The snapshot actor was already allocated, if some other snapshot actor is referencing it
+			const TOptional<AActor*> PreallocatedSnapshotVersion = SavedData->GetPreallocatedIfValidButDoNotAllocate();
+			return PreallocatedSnapshotVersion.Get(nullptr) == SnapshotPropertyValue ? Equivalent : NotEquivalent;
 		}
-		else
+
+		return NotComparable;
+	}
+
+	EEquivalenceResult AreSubobjectsEquivalent(UObject* SnapshotPropertyValue, UObject* OriginalPropertyValue, const TMap<FSoftObjectPath, FActorSnapshotData>& ActorData)
+	{
+		const bool bIsWorldObject = SnapshotPropertyValue->IsInA(UWorld::StaticClass());
+		if (bIsWorldObject)
 		{
-			// Example /Game/MapName.MapName:PersistentLevel.StaticMeshActor_42.StaticMeshComponent
-			bIsPathToActorSubobject = true;
-			// Converts it to /Game/MapName.MapName:PersistentLevel.StaticMeshActor_42
-			const FSoftObjectPath PathToActor(OriginalObjectPath.GetAssetPathName(), SubPathString.Left(DotAfterActorNameIndex));
-			
-			FActorSnapshotData* Result = ActorData.Find(PathToActor);
-			UE_CLOG(Result == nullptr, LogLevelSnapshots, Warning, TEXT("Path %s looks like an actor path but no data was saved for it. Maybe it was a reference to an auto-generated actor, e.g. a brush or volume present in all worlds by default?"), *OriginalObjectPath.ToString());
-			return Result ? Result : TOptional<FActorSnapshotData*>();  
+			AActor* SnapshotOwningActor = SnapshotPropertyValue->GetTypedOuter<AActor>();
+			AActor* OriginalOwningActor = OriginalPropertyValue->GetTypedOuter<AActor>();
+			if (!ensureMsgf(SnapshotOwningActor && OriginalOwningActor, TEXT("This is weird: the objects are part of a world and not actors, so they should be subobjects of actors, like components. Investigate")))
+			{
+				return NotEquivalent;
+			}
+
+			// Are the two subobjects owned by corresponding actors
+			const FActorSnapshotData* CorrespondingSnapshotActor = ActorData.Find(OriginalOwningActor);
+			const bool bAreOwnedByEquivalentActors = CorrespondingSnapshotActor == nullptr || CorrespondingSnapshotActor->GetPreallocatedIfValidButDoNotAllocate().Get(nullptr) != SnapshotOwningActor; 
+			if (!bAreOwnedByEquivalentActors)
+			{
+				return NotEquivalent;
+			}
+
+			// TODO: Call registered callbacks from external modules to determine whether the two components match each other.
+			// Needed for cases like Foliage where components are matched based on foliage type instead of name
+
+			// Check that chain of outers correspond to each other.
+			UObject* CurrentSnapshotOuter = SnapshotPropertyValue;
+			UObject* CurrentOriginalOuter = OriginalPropertyValue;
+			for (; CurrentSnapshotOuter != SnapshotOwningActor && CurrentOriginalOuter != OriginalOwningActor; CurrentSnapshotOuter = CurrentSnapshotOuter->GetOuter(), CurrentOriginalOuter = CurrentOriginalOuter->GetOuter())
+			{
+				// TODO: Call registered callbacks from external modules to determine whether the two components match each other.
+				// Needed for cases like Foliage where components are matched based on foliage type instead of name
+
+				const bool bHaveSameName = CurrentSnapshotOuter->GetFName().IsEqual(CurrentOriginalOuter->GetFName());
+				// I thought of also checking whether the two outers have the same class but I see no reason to atm
+				if (!bHaveSameName)
+				{
+					return NotEquivalent;
+				}
+			}
+
+			return Equivalent;
 		}
+		
+		return NotComparable;
 	}
 }
 
@@ -99,7 +200,7 @@ void FWorldSnapshotData::SnapshotWorld(UWorld* World)
 	}
 }
 
-void FWorldSnapshotData::ApplyToWorld(UWorld* WorldToApplyTo, const FPropertySelectionMap& PropertiesToSerialize)
+void FWorldSnapshotData::ApplyToWorld(UWorld* WorldToApplyTo, UPackage* LocalisationSnapshotPackage, const FPropertySelectionMap& PropertiesToSerialize)
 {
 	DECLARE_SCOPE_CYCLE_COUNTER(TEXT("ApplySnapshotToWorld"), STAT_ApplySnapshotToWorld, STATGROUP_LevelSnapshots);
 	if (WorldToApplyTo == nullptr)
@@ -123,13 +224,13 @@ void FWorldSnapshotData::ApplyToWorld(UWorld* WorldToApplyTo, const FPropertySel
 #if WITH_EDITOR
 	ApplyToWorldTask.EnterProgressFrame(NumActorsToRecreate);
 #endif
-	ApplyToWorld_HandleRecreatingActors(EvaluatedActors, PropertiesToSerialize);	
+	ApplyToWorld_HandleRecreatingActors(EvaluatedActors, LocalisationSnapshotPackage, PropertiesToSerialize);	
 
 
 #if WITH_EDITOR
 	ApplyToWorldTask.EnterProgressFrame(NumMatchingActors);
 #endif
-	ApplyToWorld_HandleSerializingMatchingActors(EvaluatedActors, SelectedPaths, PropertiesToSerialize);
+	ApplyToWorld_HandleSerializingMatchingActors(EvaluatedActors, SelectedPaths, LocalisationSnapshotPackage, PropertiesToSerialize);
 
 
 	
@@ -160,46 +261,57 @@ bool FWorldSnapshotData::HasMatchingSavedActor(const FSoftObjectPath& OriginalOb
 	return ActorData.Contains(OriginalObjectPath);
 }
 
-TOptional<AActor*> FWorldSnapshotData::GetDeserializedActor(const FSoftObjectPath& OriginalObjectPath)
+TOptional<AActor*> FWorldSnapshotData::GetDeserializedActor(const FSoftObjectPath& OriginalObjectPath, UPackage* LocalisationSnapshotPackage)
 {
 	FActorSnapshotData* SerializedActor = ActorData.Find(OriginalObjectPath);
 	if (SerializedActor)
 	{
-		return SerializedActor->GetDeserialized(TempActorWorld->GetWorld(), *this);
+		return SerializedActor->GetDeserialized(TempActorWorld->GetWorld(), *this, LocalisationSnapshotPackage);
 	}
 
 	UE_LOG(LogLevelSnapshots, Warning, TEXT("No save data found for actor %s"), *OriginalObjectPath.ToString());
 	return {};
 }
 
-bool FWorldSnapshotData::AreReferencesEquivalent(UObject* SnapshotVersion, UObject* OriginalVersion) const
+TOptional<FObjectSnapshotData*> FWorldSnapshotData::GetSerializedClassDefaults(UClass* Class)
 {
-	if (SnapshotVersion == nullptr || OriginalVersion == nullptr)
+	FObjectSnapshotData* ClassDefaultData = ClassDefaults.Find(Class);
+	return ClassDefaultData ? ClassDefaultData : TOptional<FObjectSnapshotData*>();
+}
+
+bool FWorldSnapshotData::AreReferencesEquivalent(UObject* SnapshotPropertyValue, UObject* OriginalPropertyValue, AActor* SnapshotActor, AActor* OriginalActor) const
+{
+	const bool bOriginalIsSubobject = OriginalPropertyValue != nullptr && OriginalPropertyValue->IsIn(OriginalActor);
+	const bool bNeedsUnsupportedSubobjectRestorationFeature = SnapshotPropertyValue == nullptr && bOriginalIsSubobject;
+	if (bNeedsUnsupportedSubobjectRestorationFeature)
 	{
-		return SnapshotVersion == OriginalVersion;
+		UE_LOG(LogLevelSnapshots, Verbose, TEXT("Object '%s' in of actor '%s' seems to be a subobject. Snapshots currently do not support re-creating subobjects."), *OriginalPropertyValue->GetName(), *OriginalActor->GetName());
+		return true;
 	}
-
-	if (SnapshotVersion->GetClass() != OriginalVersion->GetClass())
+	
+	if (SnapshotPropertyValue == nullptr || OriginalPropertyValue == nullptr)
+	{
+		return SnapshotPropertyValue == OriginalPropertyValue;
+	}
+	if (SnapshotPropertyValue->GetClass() != OriginalPropertyValue->GetClass())
 	{
 		return false;
 	}
 
-	if (AActor* OriginalActorReference = Cast<AActor>(OriginalVersion))
+	const EEquivalenceResult ActorsEquivalent = AreActorsEquivalent(SnapshotPropertyValue, OriginalPropertyValue, ActorData);
+	if (ActorsEquivalent != NotComparable)
 	{
-		const FActorSnapshotData* SavedData = ActorData.Find(OriginalActorReference);
-		if (SavedData == nullptr)
-		{
-			return false;
-		}
-
-		// The snapshot actor was already allocated, if some other snapshot actor is referencing it
-		const TOptional<AActor*> PreallocatedSnapshotVersion = SavedData->GetPreallocatedIfValidButDoNotAllocate();
-		return PreallocatedSnapshotVersion.IsSet() && PreallocatedSnapshotVersion.GetValue() == SnapshotVersion;
+		return ActorsEquivalent == Equivalent;
 	}
 
-	// TODO: We can theoretically handle when OriginalVersion is a component reference
-	// TODO: Handle subobjects here when subobject support is implemented
-	return false;
+	const EEquivalenceResult SubobjectsEquivalent = AreSubobjectsEquivalent(SnapshotPropertyValue, OriginalPropertyValue, ActorData);
+	if (SubobjectsEquivalent != NotComparable)
+	{
+		return SubobjectsEquivalent == Equivalent;
+	}
+	
+	// Compare external references, like a UMaterial in the content browser.
+	return SnapshotPropertyValue == OriginalPropertyValue;
 }
 
 void FWorldSnapshotData::ApplyToWorld_HandleRemovingActors(const FPropertySelectionMap& PropertiesToSerialize)
@@ -213,7 +325,7 @@ void FWorldSnapshotData::ApplyToWorld_HandleRemovingActors(const FPropertySelect
 	}
 }
 
-void FWorldSnapshotData::ApplyToWorld_HandleRecreatingActors(TSet<AActor*>& EvaluatedActors, const FPropertySelectionMap& PropertiesToSerialize)
+void FWorldSnapshotData::ApplyToWorld_HandleRecreatingActors(TSet<AActor*>& EvaluatedActors, UPackage* LocalisationSnapshotPackage, const FPropertySelectionMap& PropertiesToSerialize)
 {
 #if WITH_EDITOR
 	FScopedSlowTask RecreateActors(PropertiesToSerialize.GetDeletedActorsToRespawn().Num(), LOCTEXT("ApplyToWorld.RecreateActorsKey", "Re-creating actors"));
@@ -291,13 +403,13 @@ void FWorldSnapshotData::ApplyToWorld_HandleRecreatingActors(TSet<AActor*>& Eval
 			{
 				// Mark it, otherwise we'll serialize it again when we look for world actors matching the snapshot
 				EvaluatedActors.Add(*RecreatedActor);
-				ActorSnapshot->DeserializeIntoRecreatedWorldActor(TempActorWorld.Get(), *RecreatedActor, *this);
+				ActorSnapshot->DeserializeIntoRecreatedEditorWorldActor(TempActorWorld.Get(), *RecreatedActor, *this, LocalisationSnapshotPackage, PropertiesToSerialize);
 			}
 		}
 	}
 }
 
-void FWorldSnapshotData::ApplyToWorld_HandleSerializingMatchingActors(TSet<AActor*>& EvaluatedActors, const TArray<FSoftObjectPath>& SelectedPaths, const FPropertySelectionMap& PropertiesToSerialize)
+void FWorldSnapshotData::ApplyToWorld_HandleSerializingMatchingActors(TSet<AActor*>& EvaluatedActors, const TArray<FSoftObjectPath>& SelectedPaths, UPackage* LocalisationSnapshotPackage, const FPropertySelectionMap& PropertiesToSerialize)
 {
 #if WITH_EDITOR
 	FScopedSlowTask ExitingActorTask(SelectedPaths.Num(), LOCTEXT("ApplyToWorld.MatchingPropertiesKey", "Writing existing actors"));
@@ -335,7 +447,7 @@ void FWorldSnapshotData::ApplyToWorld_HandleSerializingMatchingActors(TSet<AActo
 				{
 					if (FActorSnapshotData* ActorSnapshot = ActorData.Find(OriginalWorldActor))
 					{
-						ActorSnapshot->DeserializeIntoExistingWorldActor(TempActorWorld.Get(), OriginalWorldActor, *this, PropertiesToSerialize);
+						ActorSnapshot->DeserializeIntoExistingWorldActor(TempActorWorld.Get(), OriginalWorldActor, *this, LocalisationSnapshotPackage, PropertiesToSerialize);
 					}
 					
 					EvaluatedActors.Add(OriginalWorldActor);
@@ -347,44 +459,118 @@ void FWorldSnapshotData::ApplyToWorld_HandleSerializingMatchingActors(TSet<AActo
 
 int32 FWorldSnapshotData::AddObjectDependency(UObject* ReferenceFromOriginalObject)
 {
+	// TODO: Check whether subobject and track as such. Do this by walking up the chain of outers and determining whether it is owned by an actor or component.
 	return SerializedObjectReferences.AddUnique(ReferenceFromOriginalObject);
 }
 
-TOptional<UObject*> FWorldSnapshotData::ResolveObjectDependency(int32 ObjectPathIndex, EResolveType ResolveType)
+UObject* FWorldSnapshotData::ResolveObjectDependencyForSnapshotWorld(int32 ObjectPathIndex)
 {
 	if (!ensure(SerializedObjectReferences.IsValidIndex(ObjectPathIndex)))
 	{
-		return {};
+		return nullptr;
 	}
 	
 	const FSoftObjectPath& OriginalObjectPath = SerializedObjectReferences[ObjectPathIndex];
-	if (ResolveType == EResolveType::ResolveForUseInOriginalWorld)
+	bool bIsPathToActorSubobject;
+	const TOptional<FActorSnapshotData*> SnapshotData = FindSavedActorDataUsingObjectPath(ActorData, OriginalObjectPath, bIsPathToActorSubobject);
+	
+	const bool bIsExternalAssetReference = !SnapshotData.IsSet();
+	if (bIsExternalAssetReference)
+	{
+		// We're dealing with an external (asset) reference, e.g. a UMaterial.
+		UObject* ExternalReference = OriginalObjectPath.ResolveObject();
+		ensureAlwaysMsgf(ExternalReference == nullptr || (!ExternalReference->IsA<AActor>() && !ExternalReference->IsA<UActorComponent>()), TEXT("Something is wrong. We just checked that the reference is not a world object but it is an actor or component."));
+		return ExternalReference;
+	}
+
+	TOptional<AActor*> Result = SnapshotData.GetValue()->GetPreallocated(TempActorWorld->GetWorld(), *this);
+	if (!Result.IsSet())
+	{
+		return nullptr;
+	}
+
+	// Reference to actor
+	AActor* SnapshotActor = *Result;
+	const bool bIsPathToActor = !bIsPathToActorSubobject; 
+	if (bIsPathToActor)
 	{
-		UObject* Result = OriginalObjectPath.ResolveObject();
-		return Result ? Result : TOptional<UObject*>();
+		return SnapshotActor;
 	}
-	checkf(ResolveType == EResolveType::ResolveForUseInTempWorld, TEXT("Did you add a new entry to EResolveType?"));
+
+	// Reference to default subobjects, such as components, or an object we already allocated.
+	const FSoftObjectPath PathToSubobject = SetActorInPath(SnapshotActor, OriginalObjectPath);
+	if (UObject* DefaultSubobject = PathToSubobject.ResolveObject())
+	{
+		return DefaultSubobject;
+	}
+
+	// TODO: The subobject is instanced. Recreate it here using Subobjects.Find(ObjectPathIndex) to recursively construct the subobject's outers
+	return nullptr;
+}
+
+UObject* FWorldSnapshotData::ResolveObjectDependencyForEditorWorld(int32 ObjectPathIndex, const FPropertySelectionMap& SelectionMap)
+{
+	if (!ensure(SerializedObjectReferences.IsValidIndex(ObjectPathIndex)))
+	{
+		return nullptr;
+	}
+	
+	const FSoftObjectPath& OriginalObjectPath = SerializedObjectReferences[ObjectPathIndex];
 
 	bool bIsPathToActorSubobject;
-	const TOptional<FActorSnapshotData*> SnapshotData = FindSavedActorDataUsingObjectPath(ActorData, OriginalObjectPath, bIsPathToActorSubobject); 
+	const TOptional<FActorSnapshotData*> SnapshotData = FindSavedActorDataUsingObjectPath(ActorData, OriginalObjectPath, bIsPathToActorSubobject);
+	UObject* ResolvedObject = OriginalObjectPath.ResolveObject();
+	
+	// Dealing with an external (asset) reference, e.g. a UMaterial?
 	const bool bIsExternalAssetReference = !SnapshotData.IsSet();
 	if (bIsExternalAssetReference)
 	{
-		// We're dealing with an external (asset) reference, e.g. a UMaterial.
-		UObject* Result = OriginalObjectPath.ResolveObject();
-		return Result ? Result : TOptional<UObject*>();
+		ensureAlwaysMsgf(ResolvedObject == nullptr || (!ResolvedObject->IsA<AActor>() && !ResolvedObject->IsA<UActorComponent>()), TEXT("Something is wrong. We just checked that the reference is not a world object but it is an actor or component."));
+		return ResolvedObject;
 	}
 
-	if (bIsPathToActorSubobject)
+	// Actor reference?
+	const bool bIsReferenceToActor = !bIsPathToActorSubobject;
+	if (bIsReferenceToActor)
+	{
+		return ResolvedObject;
+	}
+
+
+
+	const FSubobjectSnapshotData* SubobjectData = Subobjects.Find(ObjectPathIndex);
+	if (!SubobjectData)
 	{
-		UE_LOG(LogLevelSnapshots, Warning, TEXT("Skipping subobject reference '%s'. The snapshot saved the subobject for future compatibility but we currently do not support restoring it. The reference will be set to null."), *OriginalObjectPath.ToString());
-		return {};
+		// This should not happen after release of 4.27. This is error may happen if certain snapshots in an older data format are applied because the data wasn't yet captured.
+		UE_LOG(LogLevelSnapshots, Warning, TEXT("Trying to resolve a subobject but found no associated object data."));
+		return ResolvedObject;
+	}
+	
+	// TODO: Implement restoring subobjects, e.g. material instances, here
+	UClass* ExpectedClass = SubobjectData->Class.ResolveClass();
+	if (ResolvedObject == nullptr || !ensure(ExpectedClass) || ExpectedClass != ResolvedObject->GetClass())
+	{
+		// There was already a subobject with the same name but it is incompatible because is a different class.
+		// TODO: Allocate a new instance, serialize all data into it, and purge ResolvedObject by replacing all references to it
 	}
 	else
 	{
-		TOptional<AActor*> Result = SnapshotData.GetValue()->GetPreallocated(TempActorWorld->GetWorld(), *this);
-		return Result ? TOptional<UObject*>(*Result) : TOptional<UObject*>();
+		// Subobject of same name already exists and is compatible
+		// TODO: Check whether subobject is marked as serialized and if not serialize using SelectionMap 	
 	}
+	
+	return ResolvedObject;
+}
+
+UObject* FWorldSnapshotData::ResolveObjectDependencyForClassDefaultObject(int32 ObjectPathIndex)
+{
+	if (!ensure(SerializedObjectReferences.IsValidIndex(ObjectPathIndex)))
+	{
+		return nullptr;
+	}
+	
+	const FSoftObjectPath& OriginalObjectPath = SerializedObjectReferences[ObjectPathIndex];
+	return OriginalObjectPath.ResolveObject();
 }
 
 int32 FWorldSnapshotData::AddSubobjectDependency(UObject* ReferenceFromOriginalObject)
@@ -413,8 +599,7 @@ void FWorldSnapshotData::AddClassDefault(UClass* Class)
 	if (!ClassDefaults.Contains(Class))
 	{
 		FObjectSnapshotData& ClassData = ClassDefaults.Add(Class);
-		FTakeClassDefaultObjectSnapshotArchive SaveClass = FTakeClassDefaultObjectSnapshotArchive::MakeArchiveForSavingClassDefaultObject(ClassData, *this);
-		Class->GetDefaultObject()->Serialize(SaveClass);
+		FTakeClassDefaultObjectSnapshotArchive::SaveClassDefaultObject(ClassData, *this, Class->GetDefaultObject());
 	}
 }
 
@@ -432,19 +617,15 @@ UObject* FWorldSnapshotData::GetClassDefault(UClass* Class)
 		return ClassDefaultData->CachedLoadedClassDefault;
 	}
 	
-	UObject* Default = NewObject<UObject>(
+	UObject* CDO = NewObject<UObject>(
 		GetTransientPackage(),
 		Class,
 		*FString("SnapshotCDO_").Append(*MakeUniqueObjectName(GetTransientPackage(), Class).ToString())
 		);
-	FTakeClassDefaultObjectSnapshotArchive RestoreArchive = FTakeClassDefaultObjectSnapshotArchive::MakeArchiveForRestoringClassDefaultObject(*ClassDefaultData, *this);
-	// TODO: Test how this affect class reinstancing
-	//Default->SetFlags(RF_ArchetypeObject); // No direct reason for this though it acts as as archetype so we should probably mark it as such
-	Default->Serialize(RestoreArchive);
+	FApplyClassDefaulDataArchive::SerializeClassDefaultObject(*ClassDefaultData, *this, CDO);
 
-	ClassDefaultData->CachedLoadedClassDefault = Default;
-	return Default;
-	
+	ClassDefaultData->CachedLoadedClassDefault = CDO;
+	return CDO;
 }
 
 void FWorldSnapshotData::SerializeClassDefaultsInto(UObject* Object)
@@ -452,8 +633,7 @@ void FWorldSnapshotData::SerializeClassDefaultsInto(UObject* Object)
 	FClassDefaultObjectSnapshotData* ClassDefaultData = ClassDefaults.Find(Object->GetClass());
 	if (ClassDefaultData)
 	{
-		FTakeClassDefaultObjectSnapshotArchive RestoreArchive = FTakeClassDefaultObjectSnapshotArchive::MakeArchiveForRestoringClassDefaultObject(*ClassDefaultData, *this);
-		Object->Serialize(RestoreArchive);
+		FApplyClassDefaulDataArchive::RestoreChangedClassDefaults(*ClassDefaultData, *this, Object);
 	}
 	else
 	{
diff --git a/Engine/Plugins/VirtualProduction/LevelSnapshots/Source/LevelSnapshots/Private/PropertyInfoHelpers.cpp b/Engine/Plugins/VirtualProduction/LevelSnapshots/Source/LevelSnapshots/Private/PropertyInfoHelpers.cpp
--- a/Engine/Plugins/VirtualProduction/LevelSnapshots/Source/LevelSnapshots/Private/PropertyInfoHelpers.cpp
+++ b/Engine/Plugins/VirtualProduction/LevelSnapshots/Source/LevelSnapshots/Private/PropertyInfoHelpers.cpp
@@ -24,8 +24,7 @@ bool FPropertyInfoHelpers::IsPropertyInContainer(const FProperty* Property)
 {
 	if (ensure(Property))
 	{
-		const FProperty* ParentProperty = GetParentProperty(Property);
-		return ParentProperty && IsPropertyContainer(ParentProperty);
+		return IsPropertyInCollection(Property) || IsPropertyInStruct(Property);
 	}
 	return false;
 }
@@ -47,9 +46,14 @@ bool FPropertyInfoHelpers::IsPropertyInStruct(const FProperty* Property)
 		return false;
 	}
 
-	FProperty* ParentProperty = GetParentProperty(Property);
+	// Parent struct could be FProperty or UScriptStruct
+
+	if (FProperty* ParentProperty = GetParentProperty(Property))
+	{
+		return ParentProperty->IsA(FStructProperty::StaticClass());
+	}
 
-	return ParentProperty && ParentProperty->IsA(FStructProperty::StaticClass());
+	return IsValid(Property->GetOwner<UScriptStruct>());
 }
 
 bool FPropertyInfoHelpers::IsPropertyInMap(const FProperty* Property)
diff --git a/Engine/Plugins/VirtualProduction/LevelSnapshots/Source/LevelSnapshots/Private/PropertySelectionMap.cpp b/Engine/Plugins/VirtualProduction/LevelSnapshots/Source/LevelSnapshots/Private/PropertySelectionMap.cpp
--- a/Engine/Plugins/VirtualProduction/LevelSnapshots/Source/LevelSnapshots/Private/PropertySelectionMap.cpp
+++ b/Engine/Plugins/VirtualProduction/LevelSnapshots/Source/LevelSnapshots/Private/PropertySelectionMap.cpp
@@ -54,6 +54,11 @@ TArray<FSoftObjectPath> FPropertySelectionMap::GetKeys() const
 	return Result;
 }
 
+int32 FPropertySelectionMap::GetKeyCount() const
+{
+	return SelectedWorldObjectsToSelectedProperties.Num();
+}
+
 const TSet<FSoftObjectPath>& FPropertySelectionMap::GetDeletedActorsToRespawn() const
 {
 	return DeletedActorsToRespawn;
diff --git a/Engine/Plugins/VirtualProduction/LevelSnapshots/Source/LevelSnapshots/Private/Restorability/SnapshotRestorability.cpp b/Engine/Plugins/VirtualProduction/LevelSnapshots/Source/LevelSnapshots/Private/Restorability/SnapshotRestorability.cpp
--- a/Engine/Plugins/VirtualProduction/LevelSnapshots/Source/LevelSnapshots/Private/Restorability/SnapshotRestorability.cpp
+++ b/Engine/Plugins/VirtualProduction/LevelSnapshots/Source/LevelSnapshots/Private/Restorability/SnapshotRestorability.cpp
@@ -103,7 +103,7 @@ bool FSnapshotRestorability::IsActorDesirableForCapture(const AActor* Actor)
 	return IsValid(Actor)
             && DoesActorHaveSupportedClass(Actor)
             && !Actor->IsTemplate()								// Should never happen, but we never want CDOs
-            // TODO: Temp test for USD && !Actor->HasAnyFlags(RF_Transient)				// Don't add transient actors in non-play worlds		
+            && !Actor->HasAnyFlags(RF_Transient)				// Don't add transient actors in non-play worlds		
 #if WITH_EDITOR
             && Actor->IsEditable()
             && Actor->IsListedInSceneOutliner() 				// Only add actors that are allowed to be selected and drawn in editor
diff --git a/Engine/Plugins/VirtualProduction/LevelSnapshots/Source/LevelSnapshots/Private/Restorability/StaticMeshCollisionPropertyComparer.cpp b/Engine/Plugins/VirtualProduction/LevelSnapshots/Source/LevelSnapshots/Private/Restorability/StaticMeshCollisionPropertyComparer.cpp
--- a/Engine/Plugins/VirtualProduction/LevelSnapshots/Source/LevelSnapshots/Private/Restorability/StaticMeshCollisionPropertyComparer.cpp
+++ b/Engine/Plugins/VirtualProduction/LevelSnapshots/Source/LevelSnapshots/Private/Restorability/StaticMeshCollisionPropertyComparer.cpp
@@ -51,7 +51,7 @@ IPropertyComparer::EPropertyComparison FStaticMeshCollisionPropertyComparer::Sho
 		const bool bUseDefaultCollision = SnapshotObject->bUseDefaultCollision;
 		return bDefaultCollisionConfigIsEqual
 			?
-			bUseDefaultCollision ? EPropertyComparison::TreatEqual : EPropertyComparison::TreatUnequal
+			bUseDefaultCollision ? EPropertyComparison::TreatEqual : EPropertyComparison::CheckNormally
 			:
 			EPropertyComparison::CheckNormally;
 	}
diff --git a/Engine/Plugins/VirtualProduction/LevelSnapshots/Source/LevelSnapshots/Public/Data/ActorSnapshotData.h b/Engine/Plugins/VirtualProduction/LevelSnapshots/Source/LevelSnapshots/Public/Data/ActorSnapshotData.h
--- a/Engine/Plugins/VirtualProduction/LevelSnapshots/Source/LevelSnapshots/Public/Data/ActorSnapshotData.h
+++ b/Engine/Plugins/VirtualProduction/LevelSnapshots/Source/LevelSnapshots/Public/Data/ActorSnapshotData.h
@@ -9,6 +9,7 @@
 
 class UActorComponent;
 class ULevelSnapshotSelectionSet;
+class UPackage;
 struct FPropertySelectionMap;
 struct FWorldSnapshotData;
 
@@ -19,12 +20,12 @@ struct LEVELSNAPSHOTS_API FActorSnapshotData
 
 	static FActorSnapshotData SnapshotActor(AActor* OriginalActor, FWorldSnapshotData& WorldData);
 	
-	void DeserializeIntoExistingWorldActor(UWorld* SnapshotWorld, AActor* OriginalActor, FWorldSnapshotData& WorldData, const FPropertySelectionMap& SelectedProperties);
-	void DeserializeIntoRecreatedWorldActor(UWorld* SnapshotWorld, AActor* OriginalActor, FWorldSnapshotData& WorldData);
+	void DeserializeIntoExistingWorldActor(UWorld* SnapshotWorld, AActor* OriginalActor, FWorldSnapshotData& WorldData, UPackage* InLocalisationSnapshotPackage, const FPropertySelectionMap& SelectedProperties);
+	void DeserializeIntoRecreatedEditorWorldActor(UWorld* SnapshotWorld, AActor* OriginalActor, FWorldSnapshotData& WorldData, UPackage* InLocalisationSnapshotPackage, const FPropertySelectionMap& SelectedProperties);
 
 	TOptional<AActor*> GetPreallocatedIfValidButDoNotAllocate() const;
 	TOptional<AActor*> GetPreallocated(UWorld* SnapshotWorld, FWorldSnapshotData& WorldData) const;
-	TOptional<AActor*> GetDeserialized(UWorld* SnapshotWorld, FWorldSnapshotData& WorldData);
+	TOptional<AActor*> GetDeserialized(UWorld* SnapshotWorld, FWorldSnapshotData& WorldData, UPackage* InLocalisationSnapshotPackage);
 
 	FSoftClassPath GetActorClass() const;
 
@@ -33,7 +34,7 @@ private:
 	using FSerializeActor = TFunction<void(AActor* OriginalActor, AActor* DerserializedActor)>;
 	using FSerializeComponent = TFunction<void(FObjectSnapshotData& SerializedCompData, FComponentSnapshotData& CompData, UActorComponent* Original, UActorComponent* Deserialized)>;
 	
-	void DeserializeIntoWorldActor(UWorld* SnapshotWorld, AActor* OriginalActor, FWorldSnapshotData& WorldData, FSerializeActor SerializeActor, FSerializeComponent SerializeComponent);
+	void DeserializeIntoWorldActor(UWorld* SnapshotWorld, AActor* OriginalActor, FWorldSnapshotData& WorldData, UPackage* InLocalisationSnapshotPackage, FSerializeActor SerializeActor, FSerializeComponent SerializeComponent);
 	void DeserializeComponents(AActor* IntoActor, FWorldSnapshotData& WorldData, TFunction<void(FObjectSnapshotData& SerializedCompData, FComponentSnapshotData&,UActorComponent*, FWorldSnapshotData&)>&& Callback);
 	
 	/* We cache the actor to avoid recreating it all the time */
diff --git a/Engine/Plugins/VirtualProduction/LevelSnapshots/Source/LevelSnapshots/Public/Data/ComponentSnapshotData.h b/Engine/Plugins/VirtualProduction/LevelSnapshots/Source/LevelSnapshots/Public/Data/ComponentSnapshotData.h
--- a/Engine/Plugins/VirtualProduction/LevelSnapshots/Source/LevelSnapshots/Public/Data/ComponentSnapshotData.h
+++ b/Engine/Plugins/VirtualProduction/LevelSnapshots/Source/LevelSnapshots/Public/Data/ComponentSnapshotData.h
@@ -7,6 +7,7 @@
 #include "ObjectSnapshotData.h"
 #include "ComponentSnapshotData.generated.h"
 
+class UPackage;
 struct FWorldSnapshotData;
 struct FPropertySelection;
 
@@ -19,10 +20,7 @@ struct LEVELSNAPSHOTS_API FComponentSnapshotData
 	
 	bool IsRestoreSupportedForSavedComponent();
 	
-	void DeserializeIntoTransient(FObjectSnapshotData& SerializedComponentData, UActorComponent* ComponentToDeserializeInto, FWorldSnapshotData& WorldData);
-	void DeserializeIntoExistingWorldActor(FObjectSnapshotData& SerializedComponentData, UActorComponent* OriginalComponentToDeserializeInto, UActorComponent* DeserializedComponentCounterpart, FWorldSnapshotData& WorldData, const FPropertySelection& PropertySelection);
-	void DeserializeIntoRecreatedWorldActor(FObjectSnapshotData& SerializedComponentData, UActorComponent* OriginalComponentToDeserializeInto, UActorComponent* DeserializedComponentCounterpart, FWorldSnapshotData& WorldData);
-
+	void DeserializeIntoTransient(FObjectSnapshotData& SerializedComponentData, UActorComponent* ComponentToDeserializeInto, FWorldSnapshotData& WorldData, UPackage* InLocalisationSnapshotPackage);
 	/* Describes how the component was created */
 	UPROPERTY()
 	EComponentCreationMethod CreationMethod;
diff --git a/Engine/Plugins/VirtualProduction/LevelSnapshots/Source/LevelSnapshots/Public/Data/WorldSnapshotData.h b/Engine/Plugins/VirtualProduction/LevelSnapshots/Source/LevelSnapshots/Public/Data/WorldSnapshotData.h
--- a/Engine/Plugins/VirtualProduction/LevelSnapshots/Source/LevelSnapshots/Public/Data/WorldSnapshotData.h
+++ b/Engine/Plugins/VirtualProduction/LevelSnapshots/Source/LevelSnapshots/Public/Data/WorldSnapshotData.h
@@ -9,10 +9,10 @@
 #include "SubobjectSnapshotData.h"
 #include "WorldSnapshotData.generated.h"
 
+class AActor;
 class FSnapshotArchive;
-class FTakeWorldObjectSnapshotArchive;
-class ULevelSnapshotSelectionSet;
 struct FPropertySelectionMap;
+struct FActorSnapshotData;
 
 /* Holds saved world data and handles all logic related to writing to the existing world. */
 USTRUCT()
@@ -20,9 +20,7 @@ struct LEVELSNAPSHOTS_API FWorldSnapshotData
 {
 	GENERATED_BODY()
 	friend FSnapshotArchive;
-	friend FTakeWorldObjectSnapshotArchive;
 	friend FActorSnapshotData;
-	friend FComponentSnapshotData;
 
 	void OnCreateSnapshotWorld(UWorld* NewTempActorWorld);
 	void OnDestroySnapshotWorld();
@@ -30,49 +28,63 @@ struct LEVELSNAPSHOTS_API FWorldSnapshotData
 	/* Records the actor in this snapshot */
 	void SnapshotWorld(UWorld* World);
 	/* Applies the saved properties to WorldActor */
-	void ApplyToWorld(UWorld* WorldToApplyTo, const FPropertySelectionMap& PropertiesToSerialize);
+	void ApplyToWorld(UWorld* WorldToApplyTo, UPackage* LocalisationSnapshotPackage, const FPropertySelectionMap& PropertiesToSerialize);
 
 	int32 GetNumSavedActors() const;
 	void ForEachOriginalActor(TFunction<void(const FSoftObjectPath& ActorPath)> HandleOriginalActorPath) const;
 	bool HasMatchingSavedActor(const FSoftObjectPath& OriginalObjectPath) const;
-	TOptional<AActor*> GetDeserializedActor(const FSoftObjectPath& OriginalObjectPath);
+
+	/* Given a path to a world actor, gets the equivalent allocated actor. */
+	TOptional<AActor*> GetDeserializedActor(const FSoftObjectPath& OriginalObjectPath, UPackage* LocalisationSnapshotPackage);
+	/* Gets the state of the CDO from when the snapshot was taken. CDO is saved when a snapshot is taken so CDO changes can be detected. */
+	TOptional<FObjectSnapshotData*> GetSerializedClassDefaults(UClass* Class);  
 
 	/**
 	 * Checks whether two pointers point to "equivalent" objects.
-	 * TODO: Currently only actors are supported. In future, extend this when subobject support is implemented.
 	 */
-	bool AreReferencesEquivalent(UObject* SnapshotVersion, UObject* OriginalVersion) const;
+	bool AreReferencesEquivalent(UObject* SnapshotPropertyValue, UObject* OriginalPropertyValue, AActor* SnapshotActor, AActor* OriginalActor) const;
 
-private:
+public: /* Serialisation functions */
 
-	void ApplyToWorld_HandleRemovingActors(const FPropertySelectionMap& PropertiesToSerialize);
-	void ApplyToWorld_HandleRecreatingActors(TSet<AActor*>& EvaluatedActors, const FPropertySelectionMap& PropertiesToSerialize);
-	void ApplyToWorld_HandleSerializingMatchingActors(TSet<AActor*>& EvaluatedActors, const TArray<FSoftObjectPath>& SelectedPaths, const FPropertySelectionMap& PropertiesToSerialize);
-	
-	enum class EResolveType
-	{
-		ResolveForUseInOriginalWorld,
-		ResolveForUseInTempWorld
-	};
+	
+	
 	/* Adds an object dependency without serializing the object's content. Intended for external objects, e.g. to UMaterial in the content browser. */
 	int32 AddObjectDependency(UObject* ReferenceFromOriginalObject);
-	TOptional<UObject*> ResolveObjectDependency(int32 ObjectPathIndex, EResolveType ResolveType);
 
+	/* Resolves an object dependency for use in the snapshot world. If the object is a subobject, it gets fully serialized. */
+	UObject* ResolveObjectDependencyForSnapshotWorld(int32 ObjectPathIndex);
+	
+	/* Resolves an object dependency for use in the editor world. If the object is a subobject, it is serialized.
+	* Steps for serializing subobject:
+	*  - If an equivalent object with the saved name and class exists, use that and serialize the properties in SelectionMap into it.
+	*  - Otherwise allocate a new object and serialize all properties into it
+	*  - Replaces all references to the trashed object with the new one
+	*/
+	UObject* ResolveObjectDependencyForEditorWorld(int32 ObjectPathIndex, const FPropertySelectionMap& SelectionMap);
+
+	/* Resolves an object depedency when restoring a class default object. Simply resolves without further checks. */
+	UObject* ResolveObjectDependencyForClassDefaultObject(int32 ObjectPathIndex);
+	
 	
 	/* Adds a subobject dependency. Intended for internal objects which need to store serialized data, e.g. components and other subobjects. Implicitly calls AddObjectDependency.
-	 * @return A valid index in SerializedObjectReferences and the corresponding subobject data.
-	 */
+	* @return A valid index in SerializedObjectReferences and the corresponding subobject data.
+	*/
 	int32 AddSubobjectDependency(UObject* ReferenceFromOriginalObject);
 
 	
 	void AddClassDefault(UClass* Class);
 	UObject* GetClassDefault(UClass* Class);
+	
 	/* Gets the Object's class and serializes the saved CDO into it.
-	 * This is intended for cases where you cannot specify a template object for new objects. Components are one such use case.
-	 */
+	* This is intended for cases where you cannot specify a template object for new objects. Components are one such use case.
+	*/
 	void SerializeClassDefaultsInto(UObject* Object);
+	
+private:
 
-
+	void ApplyToWorld_HandleRemovingActors(const FPropertySelectionMap& PropertiesToSerialize);
+	void ApplyToWorld_HandleRecreatingActors(TSet<AActor*>& EvaluatedActors, UPackage* LocalisationSnapshotPackage, const FPropertySelectionMap& PropertiesToSerialize);
+	void ApplyToWorld_HandleSerializingMatchingActors(TSet<AActor*>& EvaluatedActors, const TArray<FSoftObjectPath>& SelectedPaths, UPackage* LocalisationSnapshotPackage, const FPropertySelectionMap& PropertiesToSerialize);
 
 	
 	
diff --git a/Engine/Plugins/VirtualProduction/LevelSnapshots/Source/LevelSnapshots/Public/PropertySelectionMap.h b/Engine/Plugins/VirtualProduction/LevelSnapshots/Source/LevelSnapshots/Public/PropertySelectionMap.h
--- a/Engine/Plugins/VirtualProduction/LevelSnapshots/Source/LevelSnapshots/Public/PropertySelectionMap.h
+++ b/Engine/Plugins/VirtualProduction/LevelSnapshots/Source/LevelSnapshots/Public/PropertySelectionMap.h
@@ -29,6 +29,7 @@ struct LEVELSNAPSHOTS_API FPropertySelectionMap
 	const FPropertySelection* GetSelectedProperties(UObject* WorldObject) const;
 	const FPropertySelection* GetSelectedProperties(const FSoftObjectPath& WorldObjectPath) const;
 	TArray<FSoftObjectPath> GetKeys() const;
+	int32 GetKeyCount() const;
 	const TSet<FSoftObjectPath>& GetDeletedActorsToRespawn() const;
 	const TSet<TWeakObjectPtr<AActor>>& GetNewActorsToDespawn() const;
 
diff --git a/Engine/Plugins/VirtualProduction/LevelSnapshots/Source/LevelSnapshotsEditor/Private/Views/Results/LevelSnapshotsEditorResults.cpp b/Engine/Plugins/VirtualProduction/LevelSnapshots/Source/LevelSnapshotsEditor/Private/Views/Results/LevelSnapshotsEditorResults.cpp
--- a/Engine/Plugins/VirtualProduction/LevelSnapshots/Source/LevelSnapshotsEditor/Private/Views/Results/LevelSnapshotsEditorResults.cpp
+++ b/Engine/Plugins/VirtualProduction/LevelSnapshots/Source/LevelSnapshotsEditor/Private/Views/Results/LevelSnapshotsEditorResults.cpp
@@ -2,6 +2,8 @@
 
 #include "Views/Results/LevelSnapshotsEditorResults.h"
 
+#include "Widgets/SNullWidget.h"
+
 #include "Views/Results/SLevelSnapshotsEditorResults.h"
 
 FLevelSnapshotsEditorResults::FLevelSnapshotsEditorResults(const TSharedRef<FLevelSnapshotsEditorViewBuilder>& InBuilder)
diff --git a/Engine/Plugins/VirtualProduction/LevelSnapshots/Source/LevelSnapshotsEditor/Private/Views/Results/SLevelSnapshotsEditorResults.cpp b/Engine/Plugins/VirtualProduction/LevelSnapshots/Source/LevelSnapshotsEditor/Private/Views/Results/SLevelSnapshotsEditorResults.cpp
--- a/Engine/Plugins/VirtualProduction/LevelSnapshots/Source/LevelSnapshotsEditor/Private/Views/Results/SLevelSnapshotsEditorResults.cpp
+++ b/Engine/Plugins/VirtualProduction/LevelSnapshots/Source/LevelSnapshotsEditor/Private/Views/Results/SLevelSnapshotsEditorResults.cpp
@@ -13,6 +13,7 @@
 #include "Algo/Find.h"
 #include "DebugViewModeHelpers.h"
 #include "EditorStyleSet.h"
+#include "Framework/MultiBox/MultiBoxBuilder.h"
 #include "Framework/Notifications/NotificationManager.h"
 #include "GameFramework/Actor.h"
 #include "IDetailPropertyRow.h"
@@ -23,9 +24,10 @@
 #include "PropertyHandle.h"
 #include "SnapshotRestorability.h"
 #include "Stats/StatsMisc.h"
+#include "Widgets/DeclarativeSyntaxSupport.h"
 #include "Widgets/Images/SImage.h"
-#include "Widgets/Input/SCheckBox.h"
 #include "Widgets/Input/SButton.h"
+#include "Widgets/Input/SComboButton.h"
 #include "Widgets/Input/SSearchBox.h"
 #include "Widgets/Layout/SBox.h"
 #include "Widgets/SBoxPanel.h"
@@ -37,6 +39,58 @@
 
 #define LOCTEXT_NAMESPACE "LevelSnapshotsEditor"
 
+TWeakPtr<IPropertyRowGenerator> FRowGeneratorInfo::GetGeneratorObject() const
+{
+	return GeneratorObject;
+}
+
+void FRowGeneratorInfo::FlushReferences()
+{
+	GeneratorObject.Reset();
+}
+
+FPropertyHandleHierarchy::FPropertyHandleHierarchy(const TSharedPtr<IDetailTreeNode>& InNode, const TSharedPtr<IPropertyHandle>& InHandle)
+	: Node(InNode)
+	, Handle(InHandle)
+{
+	const FString& DefaultPropertyName = "Base";
+
+	if (Node.IsValid())
+	{
+		if (ensure(Handle.IsValid()))
+		{
+			// Check first if map
+			const TSharedPtr<IPropertyHandle> KeyHandle = Handle->GetKeyHandle();
+
+			if (KeyHandle.IsValid() && KeyHandle->IsValidHandle())
+			{
+				KeyHandle->GetValueAsFormattedString(PropertyName);
+				PropertyName = "Key: " + PropertyName;
+			}
+			else
+			{
+				PropertyName = Handle->GetPropertyDisplayName().ToString();
+			}
+
+			const TSharedPtr<IPropertyHandle> ParentHandle = Handle->GetParentHandle();
+
+			if (ParentHandle.IsValid() && ParentHandle->IsValidHandle())
+			{
+				ParentPropertyName = ParentHandle->GetPropertyDisplayName().ToString();
+			}
+			else
+			{
+				ParentPropertyName = DefaultPropertyName;
+			}
+		}
+	}
+	else
+	{
+		PropertyName = DefaultPropertyName;
+		ParentPropertyName = DefaultPropertyName;
+	}
+}
+
 FLevelSnapshotsEditorResultsRow::~FLevelSnapshotsEditorResultsRow()
 {
 	FlushReferences();
@@ -44,19 +98,37 @@ FLevelSnapshotsEditorResultsRow::~FLevelSnapshotsEditorResultsRow()
 
 void FLevelSnapshotsEditorResultsRow::FlushReferences()
 {
-	ChildRows.Empty();
-
-	HeaderColumns.Empty();
+	if (ChildRows.Num())
+	{
+		ChildRows.Empty();
+	}
 
-	PRG_WorldCustomWidget.Reset();
-	PRG_SnapshotCustomWidget.Reset();
+	if (HeaderColumns.Num())
+	{
+		HeaderColumns.Empty();
+	}
 
-	SnapshotObject.Reset();
-	WorldObject.Reset();
-	ResultsViewPtr.Reset();
+	if (SnapshotObject.IsValid())
+	{
+		SnapshotObject.Reset();
+	}
+	if (WorldObject.IsValid())
+	{
+		WorldObject.Reset();
+	}
+	if (ResultsViewPtr.IsValid())
+	{
+		ResultsViewPtr.Reset();
+	}
 
-	SnapshotPropertyHandle.Reset();
-	WorldPropertyHandle.Reset();
+	if (SnapshotPropertyHandleHierarchy.IsValid())
+	{
+		SnapshotPropertyHandleHierarchy = nullptr;
+	}
+	if (WorldPropertyHandleHierarchy.IsValid())
+	{
+		WorldPropertyHandleHierarchy = nullptr;
+	}
 }
 
 FLevelSnapshotsEditorResultsRow::FLevelSnapshotsEditorResultsRow(
@@ -104,131 +176,33 @@ void FLevelSnapshotsEditorResultsRow::InitActorRow(
 }
 
 void FLevelSnapshotsEditorResultsRow::InitObjectRow(
+	UObject* InSnapshotObject, UObject* InWorldObject,
 	const TWeakPtr<FRowGeneratorInfo>& InSnapshotRowGenerator,
 	const TWeakPtr<FRowGeneratorInfo>& InWorldRowGenerator,
 	const TWeakPtr<SLevelSnapshotsEditorResults>& InResultsView)
 {
+	SnapshotObject = InSnapshotObject;
+	WorldObject = InWorldObject;
 	SnapshotRowGeneratorInfo = InSnapshotRowGenerator;
 	WorldRowGeneratorInfo = InWorldRowGenerator;
 	ResultsViewPtr = InResultsView;
-
-	if (InSnapshotRowGenerator.IsValid())
-	{
-		const TSharedPtr<IPropertyRowGenerator> Generator = InSnapshotRowGenerator.Pin()->GetGeneratorObject().Pin();
-		
-		if (ensureAlwaysMsgf(Generator.IsValid(), 
-			TEXT("%hs: FRowGeneratorInfo for snapshot object is valid but does not contain valid IPropertyRowGenerator object."), __FUNCTION__))
-		{
-			if (ensureAlwaysMsgf(Generator->GetSelectedObjects().Num(), 
-				TEXT("%hs: FRowGeneratorInfo and generator object for snapshot object are valid but generator object has no objects selected."), __FUNCTION__))
-			{
-				SnapshotObject = Generator->GetSelectedObjects()[0];
-			}
-		}
-	}
-	
-	if (InWorldRowGenerator.IsValid())
-	{
-		const TSharedPtr<IPropertyRowGenerator> Generator = InWorldRowGenerator.Pin()->GetGeneratorObject().Pin();
-
-		if (ensureAlwaysMsgf(Generator.IsValid(),
-			TEXT("%hs: FRowGeneratorInfo for world object is valid but does not contain valid IPropertyRowGenerator object."), __FUNCTION__))
-		{
-			if (ensureAlwaysMsgf(Generator->GetSelectedObjects().Num(),
-				TEXT("%hs: FRowGeneratorInfo and generator object for world object are valid but generator object has no objects selected."), __FUNCTION__))
-			{
-				SnapshotObject = Generator->GetSelectedObjects()[0];
-			}
-		}
-	}
-}
-
-void FLevelSnapshotsEditorResultsRow::InitPropertyRowWithHandles(
-	const TWeakPtr<FLevelSnapshotsEditorResultsRow>& InContainingObjectGroup,
-	const TSharedPtr<IPropertyHandle>& InSnapshotHandle, const TSharedPtr<IPropertyHandle>& InWorldHandle,
-	const bool bNewIsCounterpartValueSame, const ELevelSnapshotsWidgetTypeCustomization InWidgetTypeCustomization)
-{
-	ContainingObjectGroup = InContainingObjectGroup;
-	SnapshotPropertyHandle = InSnapshotHandle;
-	WorldPropertyHandle = InWorldHandle;
-	bIsCounterpartValueSame = bNewIsCounterpartValueSame;
-	WidgetTypeCustomization = InWidgetTypeCustomization;
 }
 
-void FLevelSnapshotsEditorResultsRow::InitPropertyRowWithCustomizedWidgets(
+void FLevelSnapshotsEditorResultsRow::InitPropertyRow(
 	const TWeakPtr<FLevelSnapshotsEditorResultsRow>& InContainingObjectGroup,
-	const TSharedPtr<IPropertyHandle>& InSnapshotHandle, const TSharedPtr<IPropertyHandle>& InWorldHandle,
-	const TSharedPtr<SWidget>& InSnapshotWidget, const TSharedPtr<SWidget>& InWorldWidget,
-	const bool bNewIsCounterpartValueSame, const ELevelSnapshotsWidgetTypeCustomization InWidgetTypeCustomization)
+	const TSharedPtr<FPropertyHandleHierarchy>& InSnapshotHierarchy, const TSharedPtr<FPropertyHandleHierarchy>& InWorldHandleHierarchy,
+	const bool bNewIsCounterpartValueSame)
 {
 	ContainingObjectGroup = InContainingObjectGroup;
-	SnapshotPropertyHandle = InSnapshotHandle;
-	WorldPropertyHandle = InWorldHandle;
-	PRG_SnapshotCustomWidget = InSnapshotWidget;
-	PRG_WorldCustomWidget = InWorldWidget;
+	SnapshotPropertyHandleHierarchy = InSnapshotHierarchy;
+	WorldPropertyHandleHierarchy = InWorldHandleHierarchy;
 	bIsCounterpartValueSame = bNewIsCounterpartValueSame;
-	WidgetTypeCustomization = InWidgetTypeCustomization;
 }
 
 void FLevelSnapshotsEditorResultsRow::GenerateActorGroupChildren(FPropertySelectionMap& PropertySelectionMap)
 {
 	struct FLocalPropertyLooper
 	{
-		struct FPropertyHandleHierarchy
-		{
-			FPropertyHandleHierarchy(const TSharedPtr<IPropertyHandle> InHandle)
-				: Handle(InHandle)
-			{
-				const FString& DefaultPropertyName = "Base";
-
-				if (InHandle.IsValid())
-				{
-					const TSharedPtr<IPropertyHandle> KeyHandle = Handle->GetKeyHandle();
-
-					if (KeyHandle.IsValid() && KeyHandle->IsValidHandle())
-					{
-						KeyHandle->GetValueAsFormattedString(PropertyName);
-						PropertyName = "Key: " + PropertyName;
-					}
-					else
-					{
-						PropertyName = Handle->GetPropertyDisplayName().ToString();
-					}
-
-					const TSharedPtr<IPropertyHandle> ParentHandle = Handle->GetParentHandle();
-
-					if (ParentHandle.IsValid() && ParentHandle->IsValidHandle())
-					{
-						ParentPropertyName = ParentHandle->GetPropertyDisplayName().ToString();
-					}
-					else
-					{
-						ParentPropertyName = DefaultPropertyName;
-					}
-				}
-				else
-				{
-					PropertyName = DefaultPropertyName;
-					ParentPropertyName = DefaultPropertyName;
-				}
-			}
-
-			void AddCustomNodeWidget(TSharedPtr<SWidget> InNodeWidget)
-			{
-				InNodeWidget->SetCanTick(false);
-				CustomWidget = InNodeWidget;
-			}
-
-			// Used to identify counterparts
-			FString PropertyName;
-			FString ParentPropertyName;
-
-			TSharedPtr<IPropertyHandle> Handle;
-			TArray<TSharedRef<FPropertyHandleHierarchy>> DirectChildren;
-
-			TSharedPtr<SWidget> CustomWidget;
-		};
-
 		static TArray<TFieldPath<FProperty>> LoopOverProperties(
 			const TWeakPtr<FRowGeneratorInfo>& InSnapshotRowGeneratorInfo, const TWeakPtr<FRowGeneratorInfo>& InWorldRowGeneratorInfo,
 			const TWeakPtr<FLevelSnapshotsEditorResultsRow>& InDirectParentRow, const TArray<TFieldPath<FProperty>>& PropertiesThatPassFilter)
@@ -342,28 +316,9 @@ void FLevelSnapshotsEditorResultsRow::GenerateActorGroupChildren(FPropertySelect
 
 				const TWeakPtr<FLevelSnapshotsEditorResultsRow>& ContainingObjectGroup = InDirectParentRow;
 
-				if (Handle->IsCustomized())
-				{
-					TSharedPtr<SWidget> SnapshotWidget = nullptr;
-
-					if (FoundHierarchy.IsValid())
-					{
-						SnapshotWidget = FoundHierarchy.Pin()->CustomWidget;
-					}
-					
-					const TSharedPtr<SWidget> WorldWidget = InHierarchy.Pin()->CustomWidget;
-					
-					NewProperty->InitPropertyRowWithCustomizedWidgets(ContainingObjectGroup, 
-						FoundHierarchy.IsValid() && FoundHierarchy.Pin()->Handle.IsValid() ? FoundHierarchy.Pin()->Handle : nullptr, 
-						Handle ? Handle : nullptr, 
-						SnapshotWidget, WorldWidget, bIsCounterpartValueSame);
-				}
-				else
-				{
-					NewProperty->InitPropertyRowWithHandles(ContainingObjectGroup, 
-						FoundHierarchy.IsValid() && FoundHierarchy.Pin()->Handle.IsValid() ? FoundHierarchy.Pin()->Handle : nullptr, 
-						Handle ? Handle : nullptr, bIsCounterpartValueSame);
-				}
+				NewProperty->InitPropertyRow(ContainingObjectGroup, 
+						FoundHierarchy.IsValid() ? FoundHierarchy.Pin() : nullptr, 
+						InHierarchy.IsValid() ? InHierarchy.Pin() : nullptr, bIsCounterpartValueSame);
 
 				for (int32 ChildIndex = 0; ChildIndex < InHierarchy.Pin()->DirectChildren.Num(); ChildIndex++)
 				{
@@ -372,7 +327,7 @@ void FLevelSnapshotsEditorResultsRow::GenerateActorGroupChildren(FPropertySelect
 					LoopOverHandleHierarchiesAndCreateRowHierarchy(ChildHierarchy, NewProperty, PropertiesThatPassFilter, PropertyRowsGenerated, InHierarchyToSearchForCounterparts);
 				}
 
-				if (NewProperty->DoesRowRepresentGroup() && !NewProperty->GetChildRows().IsValidIndex(0))
+				if (NewProperty->DoesRowRepresentGroup() && !NewProperty->GetChildRows().Num())
 				{
 					// No valid children, destroy group
 					NewProperty.Reset();
@@ -441,16 +396,21 @@ void FLevelSnapshotsEditorResultsRow::GenerateActorGroupChildren(FPropertySelect
 
 			if (NodeType == EDetailNodeType::Item)
 			{
-				if (const TSharedPtr<IPropertyHandle> Handle = InNode->CreatePropertyHandle())
+				TSharedPtr<IPropertyHandle> Handle;
+
+				// If the handle already exists then we should just go get it
+				if (InNode->GetRow().IsValid() && InNode->GetRow()->GetPropertyHandle().IsValid())
+				{
+					Handle = InNode->GetRow()->GetPropertyHandle();
+				}
+				else // Otherwise let's try to create it
 				{
-					TSharedRef<FPropertyHandleHierarchy> NewHierarchy = MakeShared<FPropertyHandleHierarchy>(Handle);
+					Handle = InNode->CreatePropertyHandle();
+				}
 
-					if (Handle->IsCustomized())
-					{
-						FNodeWidgets Widgets = InNode->CreateNodeWidgets();
-						
-						NewHierarchy->AddCustomNodeWidget(Widgets.WholeRowWidget.IsValid() ? Widgets.WholeRowWidget : Widgets.ValueWidget);
-					}
+				if (Handle.IsValid())
+				{
+					const TSharedRef<FPropertyHandleHierarchy> NewHierarchy = MakeShared<FPropertyHandleHierarchy>(InNode, Handle);
 
 					HierarchyToPass = NewHierarchy;
 
@@ -473,7 +433,7 @@ void FLevelSnapshotsEditorResultsRow::GenerateActorGroupChildren(FPropertySelect
 			check(InRowGenerator.IsValid());
 
 			// Create a base hierarchy with dummy info and no handle
-			TSharedRef<FPropertyHandleHierarchy> ReturnHierarchy = MakeShared<FPropertyHandleHierarchy>(nullptr);
+			TSharedRef<FPropertyHandleHierarchy> ReturnHierarchy = MakeShared<FPropertyHandleHierarchy>(nullptr, nullptr);
 
 			if (InRowGenerator.Pin()->GetGeneratorObject().IsValid())
 			{
@@ -528,7 +488,7 @@ void FLevelSnapshotsEditorResultsRow::GenerateActorGroupChildren(FPropertySelect
 			check(InActor);
 			check(SceneComponent);
 			
-			TSharedRef<FComponentHierarchy> NewHierarchy = MakeShared<FComponentHierarchy>(SceneComponent);
+			const TSharedRef<FComponentHierarchy> NewHierarchy = MakeShared<FComponentHierarchy>(SceneComponent);
 
 			const int32 ReturnValue = AllHierarchies.Add(NewHierarchy);
 
@@ -594,20 +554,57 @@ void FLevelSnapshotsEditorResultsRow::GenerateActorGroupChildren(FPropertySelect
 			FPropertyEditorModule& PropertyEditorModule, const FPropertySelectionMap& PropertySelectionMap, 
 			const TWeakPtr<FLevelSnapshotsEditorResultsRow>& InDirectParentRow, const TWeakPtr<SLevelSnapshotsEditorResults>& InResultsView)
 		{
+			struct Local
+			{
+				static void CheckComponentsForVisiblePropertiesRecursively(const TWeakPtr<FComponentHierarchy>& InHierarchy, const FPropertySelectionMap& PropertySelectionMap, bool& bHasVisibleComponents)
+				{
+					USceneComponent* CurrentComponent = InHierarchy.Pin()->Component.Get();
+					
+					const FPropertySelection* PropertySelection = PropertySelectionMap.GetSelectedProperties(CurrentComponent);
+
+					bHasVisibleComponents = PropertySelection ? true : false;
+
+					if (!bHasVisibleComponents)
+					{
+						for (const TSharedRef<FComponentHierarchy>& Child : InHierarchy.Pin()->DirectChildren)
+						{
+							CheckComponentsForVisiblePropertiesRecursively(Child, PropertySelectionMap, bHasVisibleComponents);
+
+							if (bHasVisibleComponents)
+							{
+								break;
+							}
+						}
+					}
+				}
+			};
+			
 			check(InHierarchy.IsValid());
 
-			USceneComponent* CurrentComponent = InHierarchy.Pin()->Component.Get();
+			bool bShouldCreateComponentRow = false;
 
-			if (const FPropertySelection* PropertySelection = PropertySelectionMap.GetSelectedProperties(CurrentComponent))
+			// If this specific component doesn't have properties to display, we need to check the child components recursively
+			Local::CheckComponentsForVisiblePropertiesRecursively(InHierarchy, PropertySelectionMap, bShouldCreateComponentRow);
+
+			if (bShouldCreateComponentRow)
 			{
-				// Get remaining properties after filter
-				if (!PropertySelection->IsEmpty())
+				USceneComponent* CurrentComponent = InHierarchy.Pin()->Component.Get();
+				
+				if (ensureAlwaysMsgf(CurrentComponent,
+					TEXT("%hs: CurrentComponent was nullptr. Please check the InHierarchy for valid component."), __FUNCTION__))
 				{
+					const FPropertySelection* PropertySelection = PropertySelectionMap.GetSelectedProperties(CurrentComponent);
+
+					const TArray<TFieldPath<FProperty>>& PropertiesThatPassFilter =
+						PropertySelection ? PropertySelection->GetSelectedLeafProperties() : TArray<TFieldPath<FProperty>>();
+
 					const TWeakPtr<FLevelSnapshotsEditorResultsRow>& ComponentPropertyAsRow =
 						BuildComponentRow(
-							CurrentComponent, InCounterpartComponents, PropertyEditorModule, PropertySelection->GetSelectedLeafProperties(), InDirectParentRow, InResultsView);
+							CurrentComponent, InCounterpartComponents, PropertyEditorModule, PropertiesThatPassFilter, InDirectParentRow, InResultsView);
 
-					if (ComponentPropertyAsRow != nullptr)
+					if (ensureAlwaysMsgf(ComponentPropertyAsRow.IsValid(),
+						TEXT("%hs: ComponentPropertyAsRow for component '%s' was not valid but code paths hould not return null value."),
+						__FUNCTION__, *CurrentComponent->GetName()))
 					{
 						for (const TSharedRef<FComponentHierarchy>& ChildHierarchy : InHierarchy.Pin()->DirectChildren)
 						{
@@ -646,7 +643,7 @@ void FLevelSnapshotsEditorResultsRow::GenerateActorGroupChildren(FPropertySelect
 				CounterpartRowGeneratorInfo.Pin()->GetGeneratorObject().Pin()->SetObjects({ CounterpartComponent });
 			}
 			
-			NewComponentGroup->InitObjectRow(CounterpartRowGeneratorInfo, RowGeneratorInfo, InResultsView);
+			NewComponentGroup->InitObjectRow(InComponent, CounterpartComponent, CounterpartRowGeneratorInfo, RowGeneratorInfo, InResultsView);
 
 			const TArray<TFieldPath<FProperty>> PropertyRowsGenerated = FLocalPropertyLooper::LoopOverProperties(CounterpartRowGeneratorInfo, RowGeneratorInfo, NewComponentGroup, PropertiesThatPassFilter);
 
@@ -659,19 +656,9 @@ void FLevelSnapshotsEditorResultsRow::GenerateActorGroupChildren(FPropertySelect
 				}
 			}
 
-			if (NewComponentGroup->GetChildRows().IsValidIndex(0))
-			{
-				InDirectParentRow.Pin()->InsertChildRowAtIndex(NewComponentGroup);
+			InDirectParentRow.Pin()->InsertChildRowAtIndex(NewComponentGroup);
 				
-				return NewComponentGroup;
-			}
-			else
-			{
-				// No valid children, destroy group
-				NewComponentGroup.Reset();
-
-				return nullptr;
-			}
+			return NewComponentGroup;
 		}
 	};
 
@@ -930,21 +917,6 @@ TWeakPtr<FLevelSnapshotsEditorResultsRow> FLevelSnapshotsEditorResultsRow::GetCo
 	return ContainingObjectGroup;
 }
 
-FLevelSnapshotsEditorResultsRow::ELevelSnapshotsWidgetTypeCustomization FLevelSnapshotsEditorResultsRow::GetWidgetTypeCustomization() const
-{
-	return WidgetTypeCustomization;
-}
-
-const TSharedPtr<SWidget>& FLevelSnapshotsEditorResultsRow::GetWorldPrgWidget() const
-{
-	return PRG_WorldCustomWidget;
-}
-
-const TSharedPtr<SWidget>& FLevelSnapshotsEditorResultsRow::GetSnapshotPrgWidget() const
-{
-	return PRG_SnapshotCustomWidget;
-}
-
 bool FLevelSnapshotsEditorResultsRow::GetHasGeneratedChildren() const
 {
 	return bHasGeneratedChildren;
@@ -1129,14 +1101,41 @@ FLevelSnapshotPropertyChain FLevelSnapshotsEditorResultsRow::GetPropertyChain()
 	return Result;
 }
 
-const TSharedPtr<IPropertyHandle>& FLevelSnapshotsEditorResultsRow::GetSnapshotPropertyHandle() const
+TSharedPtr<IDetailTreeNode> FLevelSnapshotsEditorResultsRow::GetSnapshotPropertyNode() const
+{
+	return SnapshotPropertyHandleHierarchy.IsValid() ? SnapshotPropertyHandleHierarchy->Node : nullptr;
+}
+
+TSharedPtr<IDetailTreeNode> FLevelSnapshotsEditorResultsRow::GetWorldPropertyNode() const
+{
+	return WorldPropertyHandleHierarchy.IsValid() ? WorldPropertyHandleHierarchy->Node : nullptr;
+}
+
+ELevelSnapshotsObjectType FLevelSnapshotsEditorResultsRow::GetFirstValidPropertyNode(
+	TSharedPtr<IDetailTreeNode>& OutNode) const
+{
+	if (const TSharedPtr<IDetailTreeNode>& WorldNode = GetWorldPropertyNode())
+	{
+		OutNode = WorldNode;
+		return ObjectType_World;
+	}
+	else if (const TSharedPtr<IDetailTreeNode>& SnapshotNode = GetSnapshotPropertyNode())
+	{
+		OutNode = SnapshotNode;
+		return ObjectType_Snapshot;
+	}
+
+	return ObjectType_None;
+}
+
+TSharedPtr<IPropertyHandle> FLevelSnapshotsEditorResultsRow::GetSnapshotPropertyHandle() const
 {
-	return SnapshotPropertyHandle;
+	return SnapshotPropertyHandleHierarchy.IsValid() ? SnapshotPropertyHandleHierarchy->Handle : nullptr;
 }
 
-const TSharedPtr<IPropertyHandle>& FLevelSnapshotsEditorResultsRow::GetWorldPropertyHandle() const
+TSharedPtr<IPropertyHandle> FLevelSnapshotsEditorResultsRow::GetWorldPropertyHandle() const
 {
-	return WorldPropertyHandle;
+	return WorldPropertyHandleHierarchy.IsValid() ? WorldPropertyHandleHierarchy->Handle : nullptr;
 }
 
 ELevelSnapshotsObjectType FLevelSnapshotsEditorResultsRow::GetFirstValidPropertyHandle(TSharedPtr<IPropertyHandle>& OutHandle) const
@@ -1316,14 +1315,14 @@ void FLevelSnapshotsEditorResultsRow::GetAllCheckedChildProperties(TArray<FLevel
 	if (HasCheckedChildren())
 	{
 		for (const FLevelSnapshotsEditorResultsRowPtr& ChildRow : GetChildRows())
-		{
-			const ELevelSnapshotsEditorResultsRowType ChildRowType = ChildRow->GetRowType();
-			
+		{		
 			if (!ChildRow.IsValid())
 			{
 				continue;
 			}
 
+			const ELevelSnapshotsEditorResultsRowType ChildRowType = ChildRow->GetRowType();
+
 			if ((ChildRowType == SingleProperty || ChildRowType == CollectionGroup) && ChildRow->GetIsNodeChecked())
 			{
 				CheckedSinglePropertyNodeArray.Add(ChildRow);
@@ -1339,6 +1338,38 @@ void FLevelSnapshotsEditorResultsRow::GetAllCheckedChildProperties(TArray<FLevel
 	}
 }
 
+void FLevelSnapshotsEditorResultsRow::GetAllUncheckedChildProperties(
+	TArray<FLevelSnapshotsEditorResultsRowPtr>& UncheckedSinglePropertyNodeArray) const
+{
+	if (HasUncheckedChildren())
+	{
+		for (const FLevelSnapshotsEditorResultsRowPtr& ChildRow : GetChildRows())
+		{
+			if (!ChildRow.IsValid())
+			{
+				continue;
+			}
+			
+			const ELevelSnapshotsEditorResultsRowType ChildRowType = ChildRow->GetRowType();
+
+			if ((ChildRowType == SingleProperty || ChildRowType == CollectionGroup) && !ChildRow->GetIsNodeChecked())
+			{
+				UncheckedSinglePropertyNodeArray.Add(ChildRow);
+			}
+
+			if (ChildRowType == StructGroup)
+			{
+				if (!ChildRow->GetIsNodeChecked())
+				{
+					UncheckedSinglePropertyNodeArray.Add(ChildRow);
+				}
+
+				ChildRow->GetAllUncheckedChildProperties(UncheckedSinglePropertyNodeArray);
+			}
+		}
+	}
+}
+
 void FLevelSnapshotsEditorResultsRow::EvaluateAndSetAllParentGroupCheckedStates() const
 {
 	TWeakPtr<FLevelSnapshotsEditorResultsRow> ParentRow = GetDirectParentRow();
@@ -1389,6 +1420,8 @@ void SLevelSnapshotsEditorResults::Construct(const FArguments& InArgs, ULevelSna
 
 	OnRefreshResultsHandle = InEditorData->OnRefreshResults.AddSP(this, &SLevelSnapshotsEditorResults::RefreshResults);
 
+	FMenuBuilder ShowOptionsMenuBuilder = BuildShowOptionsMenu();
+
 	ChildSlot
 		[
 			SNew(SVerticalBox)
@@ -1418,31 +1451,40 @@ void SLevelSnapshotsEditorResults::Construct(const FArguments& InArgs, ULevelSna
 					.Text(NSLOCTEXT("LevelSnapshots", "LevelSnapshots", "Level Snapshots"))
 					.Font(FCoreStyle::GetDefaultFontStyle("Bold", 16))
 				]
-
-				+ SHorizontalBox::Slot()
-					.Padding(2.f, 0.f)
-					.VAlign(VAlign_Center)
-					.HAlign(HAlign_Right)
-					[
-
-						SNew(SButton)
-						.VAlign(VAlign_Center)
-						.HAlign(HAlign_Right)
-						.OnClicked(this, &SLevelSnapshotsEditorResults::SetAllActorGroupsCollapsed)
-						[
-							SNew(STextBlock)
-							.Text(LOCTEXT("LevelSnapshotsEditorResults_CollapseAll", "Collapse All"))
-						]
-					]
 			]	
 			
 			+ SVerticalBox::Slot()
 			.VAlign(VAlign_Top)
 			.AutoHeight()
 			[
-				SAssignNew(ResultsSearchBoxPtr, SSearchBox)
-				.HintText(LOCTEXT("LevelSnapshotsEditorResults_SearchHintText", "Search actors, components, properties..."))
-				.OnTextChanged_Raw(this, &SLevelSnapshotsEditorResults::OnResultsViewSearchTextChanged)
+				SNew(SHorizontalBox)
+
+				+SHorizontalBox::Slot()
+				[
+					SAssignNew(ResultsSearchBoxPtr, SSearchBox)
+					.HintText(LOCTEXT("LevelSnapshotsEditorResults_SearchHintText", "Search actors, components, properties..."))
+					.OnTextChanged_Raw(this, &SLevelSnapshotsEditorResults::OnResultsViewSearchTextChanged)
+				]
+
+				+SHorizontalBox::Slot()
+				.HAlign(HAlign_Right)
+				.AutoWidth()
+				[
+					SNew(SComboButton)
+					.ContentPadding(0)
+					.ForegroundColor(FSlateColor::UseForeground())
+					.ButtonStyle(FEditorStyle::Get(), "ToggleButton")
+					.AddMetaData<FTagMetaData>(FTagMetaData(TEXT("ViewOptions")))
+					.MenuContent()
+					[
+						ShowOptionsMenuBuilder.MakeWidget()
+					]
+					.ButtonContent()
+					[
+						SNew(SImage)
+						.Image(FEditorStyle::GetBrush("GenericViewButton"))
+					]
+				]
 			]
 
 			+ SVerticalBox::Slot()
@@ -1568,6 +1610,76 @@ SLevelSnapshotsEditorResults::~SLevelSnapshotsEditorResults()
 	TreeViewPtr.Reset();
 }
 
+FMenuBuilder SLevelSnapshotsEditorResults::BuildShowOptionsMenu()
+{
+	FMenuBuilder ShowOptionsMenuBuilder = FMenuBuilder(true, nullptr);
+
+	ShowOptionsMenuBuilder.AddMenuEntry(
+		LOCTEXT("ShowFilteredActors", "Show Filtered Actors"),
+		LOCTEXT("ShowFilteredActors_Tooltip", "If false, only displays actors which have passed the applied filter. Refresh Results to apply."),
+		FSlateIcon(),
+		FUIAction(
+			FExecuteAction::CreateLambda([this]() {
+				SetShowFilteredActors(!GetShowFilteredActors());
+				}),
+			FCanExecuteAction(),
+			FIsActionChecked::CreateSP(this, &SLevelSnapshotsEditorResults::GetShowFilteredActors)
+		),
+		NAME_None,
+		EUserInterfaceActionType::ToggleButton
+	);
+
+	ShowOptionsMenuBuilder.AddMenuEntry(
+		LOCTEXT("ShowUnselectedActors", "Show Unselected Actors"),
+		LOCTEXT("ShowUnselectedActors_Tooltip", "If false, unselected actors will be hidden from view."),
+		FSlateIcon(),
+		FUIAction(
+			FExecuteAction::CreateLambda([this]() {
+				SetShowUnselectedActors(!GetShowUnselectedActors());
+				}),
+			FCanExecuteAction(),
+			FIsActionChecked::CreateSP(this, &SLevelSnapshotsEditorResults::GetShowUnselectedActors)
+		),
+		NAME_None,
+		EUserInterfaceActionType::ToggleButton
+	);
+
+	ShowOptionsMenuBuilder.AddMenuEntry(
+		LOCTEXT("CollapseAll", "Collapse All"),
+		LOCTEXT("LevelSnapshotsResultsView_CollapseAll_Tooltip", "Collapse all expanded actor groups in the Modified Actors list."),
+		FSlateIcon(),
+		FUIAction(
+			FExecuteAction::CreateLambda([this]() {
+				SetAllActorGroupsCollapsed();
+				})
+		),
+		NAME_None,
+		EUserInterfaceActionType::Button
+	);
+
+	return ShowOptionsMenuBuilder;
+}
+
+void SLevelSnapshotsEditorResults::SetShowFilteredActors(const bool bNewSetting)
+{
+	bShowFilteredActors = bNewSetting;
+}
+
+void SLevelSnapshotsEditorResults::SetShowUnselectedActors(const bool bNewSetting)
+{
+	bShowUnselectedActors = bNewSetting;
+}
+
+bool SLevelSnapshotsEditorResults::GetShowFilteredActors() const
+{
+	return bShowFilteredActors;
+}
+
+bool SLevelSnapshotsEditorResults::GetShowUnselectedActors() const
+{
+	return bShowUnselectedActors;
+}
+
 void SLevelSnapshotsEditorResults::FlushMemory()
 {
 	TreeViewRootHeaderObjects.Empty();
@@ -1630,6 +1742,8 @@ FReply SLevelSnapshotsEditorResults::OnClickApplyToWorld()
 		UWorld* World = EditorDataPtr->GetSelectedWorld();
 		ActiveLevelSnapshot.GetValue()->ApplySnapshotToWorld(World, EditorDataPtr->GetFilterResults()->GetPropertiesToRollback());
 
+		SetAllActorGroupsCollapsed();
+
 		RefreshResults();
 	}
 	else
@@ -1685,68 +1799,76 @@ void SLevelSnapshotsEditorResults::BuildSelectionSetFromSelectedPropertiesInEach
 
 	struct Local
 	{
-		static void AddAllChildPropertiesInObjectGroupToSelectionSetRecursively(
+		static void RemoveUndesirablePropertiesFromSelectionMapRecursively(
 			FPropertySelectionMap& SelectionMap, const FLevelSnapshotsEditorResultsRowPtr& Group)
 		{
-			if (!ensureMsgf(Group->DoesRowRepresentObject(), 
+			if (!ensureMsgf(Group->DoesRowRepresentObject() && Group->GetWorldObject(),
 				TEXT("AddAllChildPropertiesInObjectGroupToSelectionSetRecursively: Group does not represent an object. Group name: %s"), *Group->GetDisplayName().ToString()))
 			{
 				return;
 			}
 
+			// We only want to check component or subobject groups or actor groups which have had their children generated already
 			if (Group->GetRowType() == FLevelSnapshotsEditorResultsRow::ActorGroup && !Group->GetHasGeneratedChildren())
 			{
 				return;
 			}
 
-			// If we have generated properties or this isn't an actor group, we must check which properties have been checked
-			// Remove it from the map then re-add it if it has checked children
+			// We only want to check groups which have unchecked children and remove properties which are unchecked
+			if (!Group->HasUncheckedChildren())
+			{
+				return;
+			}
+
+			// Make a copy of the property selection. If a node is unchecked, we'll remove it from the copy.
+			FPropertySelection CheckedNodeFieldPaths = *SelectionMap.GetSelectedProperties(Group->GetWorldObject());
 			SelectionMap.RemoveObjectPropertiesFromMap(Group->GetWorldObject());
+
+			TArray<FLevelSnapshotsEditorResultsRowPtr> UncheckedChildPropertyNodes;
 			
-			FPropertySelection CheckedNodeFieldPaths;
-			if (Group->HasCheckedChildren())
+			for (const FLevelSnapshotsEditorResultsRowPtr& ChildRow : Group->GetChildRows())
 			{
-				TArray<FLevelSnapshotsEditorResultsRowPtr> CheckedChildPropertyNodes;
-				
-				for (const FLevelSnapshotsEditorResultsRowPtr& ChildRow : Group->GetChildRows())
+				if (!ChildRow.IsValid())
 				{
-					if (!ChildRow.IsValid())
-					{
-						continue;
-					}
+					continue;
+				}
 
-					const FLevelSnapshotsEditorResultsRow::ELevelSnapshotsEditorResultsRowType ChildRowType = ChildRow->GetRowType();
+				const FLevelSnapshotsEditorResultsRow::ELevelSnapshotsEditorResultsRowType ChildRowType = ChildRow->GetRowType();
 
-					if (ChildRow->DoesRowRepresentObject())
-					{
-						AddAllChildPropertiesInObjectGroupToSelectionSetRecursively(SelectionMap, ChildRow);
-					}
-					else if ((ChildRowType == FLevelSnapshotsEditorResultsRow::SingleProperty || ChildRowType == FLevelSnapshotsEditorResultsRow::CollectionGroup) && 
-						ChildRow->GetIsNodeChecked())
-					{
-						CheckedChildPropertyNodes.Add(ChildRow);
-					}
-					else if (ChildRowType == FLevelSnapshotsEditorResultsRow::StructGroup)
+				if (ChildRow->DoesRowRepresentObject())
+				{
+					RemoveUndesirablePropertiesFromSelectionMapRecursively(SelectionMap, ChildRow);
+				}
+				else if ((ChildRowType == FLevelSnapshotsEditorResultsRow::SingleProperty || ChildRowType == FLevelSnapshotsEditorResultsRow::CollectionGroup) && 
+					!ChildRow->GetIsNodeChecked())
+				{
+					UncheckedChildPropertyNodes.Add(ChildRow);
+				}
+				else if (ChildRowType == FLevelSnapshotsEditorResultsRow::StructGroup)
+				{
+					if (!ChildRow->GetIsNodeChecked())
 					{
-						if (ChildRow->HasCheckedChildren())
-						{
-							CheckedChildPropertyNodes.Add(ChildRow);
-							
-							ChildRow->GetAllCheckedChildProperties(CheckedChildPropertyNodes);
-						}
+						UncheckedChildPropertyNodes.Add(ChildRow);
 					}
+					
+					ChildRow->GetAllUncheckedChildProperties(UncheckedChildPropertyNodes);
 				}
+			}
 
-				for (const FLevelSnapshotsEditorResultsRowPtr& ChildRow : CheckedChildPropertyNodes)
+			for (const FLevelSnapshotsEditorResultsRowPtr& ChildRow : UncheckedChildPropertyNodes)
+			{
+				if (ChildRow.IsValid())
 				{
-					if (ChildRow.IsValid())
-					{
-						CheckedNodeFieldPaths.AddProperty(ChildRow->GetPropertyChain());
-					}
+					FLevelSnapshotPropertyChain Chain = ChildRow->GetPropertyChain();
+					
+					CheckedNodeFieldPaths.RemoveProperty(&Chain, ChildRow->GetProperty());
 				}
 			}
 
-			SelectionMap.AddObjectProperties(Group->GetWorldObject(), CheckedNodeFieldPaths);
+			if (CheckedNodeFieldPaths.GetSelectedLeafProperties().Num())
+			{
+				SelectionMap.AddObjectProperties(Group->GetWorldObject(), CheckedNodeFieldPaths);
+			}
 		}
 	};
 
@@ -1757,7 +1879,7 @@ void SLevelSnapshotsEditorResults::BuildSelectionSetFromSelectedPropertiesInEach
 	{
 		if (Group.IsValid())
 		{
-			if (Group->GetDesiredVisibility() != EVisibility::Visible || Group->GetWidgetCheckedState() == ECheckBoxState::Unchecked)
+			if (Group->GetWidgetCheckedState() == ECheckBoxState::Unchecked)
 			{
 				// Remove from PropertyMap
 				if (AActor* WorldActor = Cast<AActor>(Group->GetWorldObject()))
@@ -1772,7 +1894,7 @@ void SLevelSnapshotsEditorResults::BuildSelectionSetFromSelectedPropertiesInEach
 			}
 			else
 			{
-				Local::AddAllChildPropertiesInObjectGroupToSelectionSetRecursively(PropertySelectionMap, Group);
+				Local::RemoveUndesirablePropertiesFromSelectionMapRecursively(PropertySelectionMap, Group);
 			}
 		}
 	}
@@ -1836,8 +1958,14 @@ TWeakPtr<FRowGeneratorInfo> SLevelSnapshotsEditorResults::RegisterRowGenerator(
 	FPropertyEditorModule& PropertyEditorModule)
 {
 	// Since we must keep many PRG objects alive in order to access the handle data, validating the nodes each tick is very taxing.
-	// We can enable validation temporarily when properties are changed if we wish
-	const TSharedRef<IPropertyRowGenerator>& RowGeneratorObject = PropertyEditorModule.CreatePropertyRowGenerator(GetLevelSnapshotsAppropriatePropertyRowGeneratorArgs());
+	// We can override the validation with a lambda since the validation function in PRG is not necessary for our implementation
+
+	auto ValidationLambda = ([](const FRootPropertyNodeList& PropertyNodeList) { return true; });
+	
+	const TSharedRef<IPropertyRowGenerator>& RowGeneratorObject = 
+		PropertyEditorModule.CreatePropertyRowGenerator(GetLevelSnapshotsAppropriatePropertyRowGeneratorArgs());
+
+	RowGeneratorObject->SetCustomValidatePropertyNodesFunction(FOnValidatePropertyRowGeneratorNodes::CreateLambda(MoveTemp(ValidationLambda)));
 	
 	const TSharedRef<FRowGeneratorInfo> NewGeneratorInfo = MakeShared<FRowGeneratorInfo>(FRowGeneratorInfo(InBoundObject, InGeneratorType, RowGeneratorObject));
 
@@ -1873,9 +2001,7 @@ FLevelSnapshotsEditorResultsRowPtr& SLevelSnapshotsEditorResults::GetOrCreateDum
 }
 
 void SLevelSnapshotsEditorResults::GenerateTreeView()
-{
-	CleanUpGenerators();
-	
+{	
 	if (!ensure(EditorDataPtr.IsValid()) || !ensure(TreeViewPtr.IsValid()))
 	{
 		return;
@@ -1883,10 +2009,7 @@ void SLevelSnapshotsEditorResults::GenerateTreeView()
 
 	DECLARE_SCOPE_CYCLE_COUNTER(TEXT("GenerateTreeView"), STAT_GenerateTreeView, STATGROUP_LevelSnapshots);
 	
-	TreeViewRootHeaderObjects.Empty();
-	TreeViewModifiedActorGroupObjects.Empty();
-	TreeViewAddedActorGroupObjects.Empty();
-	TreeViewRemovedActorGroupObjects.Empty();
+	FlushMemory();
 	
 	UFilteredResults* FilteredResults = EditorDataPtr->GetFilterResults(); 
 	const TWeakObjectPtr<ULevelSnapshotFilter>& UserFilters = FilteredResults->GetUserFilters();
@@ -1898,7 +2021,7 @@ void SLevelSnapshotsEditorResults::GenerateTreeView()
 	SplitterManagerPtr = MakeShared<FLevelSnapshotsEditorResultsSplitterManager>(FLevelSnapshotsEditorResultsSplitterManager());
 
 	// Create root headers
-	if (FilterListData.GetModifiedFilteredActors().Num() > 0)
+	if (FilterListData.GetModifiedFilteredActors().Num())
 	{
 		FLevelSnapshotsEditorResultsRowPtr ModifiedActorsHeader = MakeShared<FLevelSnapshotsEditorResultsRow>(
 			FLevelSnapshotsEditorResultsRow(FText::GetEmpty(), FLevelSnapshotsEditorResultsRow::TreeViewHeader, ECheckBoxState::Checked));
@@ -1909,73 +2032,56 @@ void SLevelSnapshotsEditorResults::GenerateTreeView()
 			},
 			SharedThis(this));
 
-		TreeViewRootHeaderObjects.Add(ModifiedActorsHeader);
-
-		GenerateTreeViewChildren_ModifiedActors(ModifiedActorsHeader, UserFilters.Get());
+		if (GenerateTreeViewChildren_ModifiedActors(ModifiedActorsHeader, UserFilters.Get()))
+		{
+			TreeViewRootHeaderObjects.Add(ModifiedActorsHeader);
+		}
+		else
+		{
+			ModifiedActorsHeader.Reset();
+		}
 	}
 
-	if (FilterListData.GetFilteredAddedWorldActors().Num() > 0)
+	if (FilterListData.GetFilteredAddedWorldActors().Num())
 	{
 		FLevelSnapshotsEditorResultsRowPtr AddedActorsHeader = MakeShared<FLevelSnapshotsEditorResultsRow>(
 			FLevelSnapshotsEditorResultsRow(FText::GetEmpty(), FLevelSnapshotsEditorResultsRow::TreeViewHeader, ECheckBoxState::Checked));
 		AddedActorsHeader->InitHeaderRow(FLevelSnapshotsEditorResultsRow::ELevelSnapshotsEditorResultsTreeViewHeaderType::HeaderType_AddedActors, 
 			{ LOCTEXT("ColumnName_AddedActors", "Added Actors") }, SharedThis(this));
 
-		TreeViewRootHeaderObjects.Add(AddedActorsHeader);
-
-		GenerateTreeViewChildren_AddedActors(AddedActorsHeader);
+		if (GenerateTreeViewChildren_AddedActors(AddedActorsHeader))
+		{
+			TreeViewRootHeaderObjects.Add(AddedActorsHeader);
+		}
+		else
+		{
+			AddedActorsHeader.Reset();
+		}
 	}
 	
-	if (FilterListData.GetFilteredRemovedOriginalActorPaths().Num() > 0)
+	if (FilterListData.GetFilteredRemovedOriginalActorPaths().Num())
 	{
 		FLevelSnapshotsEditorResultsRowPtr RemovedActorsHeader = MakeShared<FLevelSnapshotsEditorResultsRow>(
 			FLevelSnapshotsEditorResultsRow(FText::GetEmpty(), FLevelSnapshotsEditorResultsRow::TreeViewHeader, ECheckBoxState::Checked));
 		RemovedActorsHeader->InitHeaderRow(FLevelSnapshotsEditorResultsRow::ELevelSnapshotsEditorResultsTreeViewHeaderType::HeaderType_RemovedActors, 
 			{ LOCTEXT("ColumnName_RemovedActors", "Removed Actors") }, SharedThis(this));
 
-		TreeViewRootHeaderObjects.Add(RemovedActorsHeader);
-
-		GenerateTreeViewChildren_RemovedActors(RemovedActorsHeader);
+		if (GenerateTreeViewChildren_RemovedActors(RemovedActorsHeader))
+		{
+			TreeViewRootHeaderObjects.Add(RemovedActorsHeader);
+		}
+		else
+		{
+			RemovedActorsHeader.Reset();
+		}
 	}
 
 	TreeViewPtr->RequestListRefresh();
 	UpdateSnapshotInformationText();
 }
 
-void SLevelSnapshotsEditorResults::GenerateTreeViewChildren_ModifiedActors(FLevelSnapshotsEditorResultsRowPtr ModifiedActorsHeader, ULevelSnapshotFilter* UserFilters)
+bool SLevelSnapshotsEditorResults::GenerateTreeViewChildren_ModifiedActors(FLevelSnapshotsEditorResultsRowPtr ModifiedActorsHeader, ULevelSnapshotFilter* UserFilters)
 {
-	struct Local
-	{
-		/* Returns true if the actor group would have any visible properties after filtering. */
-		static bool DoesActorContainAnyVisibleProperties(TSet<FSoftObjectPath>& EvaluatedObjects, FFilterListData& FilterListData)
-		{
-			bool bHasAnyVisibleProperties = false;
-
-			const TArray<FSoftObjectPath>& SelectionMapKeys = FilterListData.GetModifiedActorsSelectedProperties().GetKeys();
-
-			for (const FSoftObjectPath& Key : SelectionMapKeys)
-			{
-				if (!EvaluatedObjects.Contains(Key))
-				{
-					EvaluatedObjects.Add(Key);
-
-					if (Key.IsValid())
-					{
-						const FPropertySelection* SelectedProperties = FilterListData.GetModifiedActorsSelectedProperties().GetSelectedProperties(Key);
-						if (!bHasAnyVisibleProperties
-							&& ensure(SelectedProperties) && !SelectedProperties->IsEmpty())
-						{
-							bHasAnyVisibleProperties = true;
-							break;
-						}
-					}
-				}
-			}
-
-			return bHasAnyVisibleProperties;
-		}
-	};
-
 	check(ModifiedActorsHeader);
 	
 	const TSet<TWeakObjectPtr<AActor>>& ActorsToConsider = FilterListData.GetModifiedFilteredActors();
@@ -2001,13 +2107,21 @@ void SLevelSnapshotsEditorResults::GenerateTreeViewChildren_ModifiedActors(FLeve
 			continue;
 		}
 
+		const int32 KeyCountBeforeFilter = FilterListData.GetModifiedActorsSelectedProperties().GetKeyCount();
+
 		// Get remaining properties after filter
-		if (UserFilters && WorldActor)
+		if (UserFilters)
 		{
 			FilterListData.ApplyFilterToFindSelectedProperties(WorldActor, UserFilters);
 		}
 
-		if (!Local::DoesActorContainAnyVisibleProperties(EvaluatedObjects, FilterListData))
+		const int32 KeyCountAfterFilter = FilterListData.GetModifiedActorsSelectedProperties().GetKeyCount();
+		const int32 KeyCountDifference = KeyCountAfterFilter - KeyCountBeforeFilter;
+
+		// If keys have been added, then this actor or its children have properties that pass the filter
+		const bool bDoesActorHavePropertiesAfterFilter = KeyCountDifference > 0;
+
+		if (!bDoesActorHavePropertiesAfterFilter)
 		{
 			continue;
 		}
@@ -2026,9 +2140,11 @@ void SLevelSnapshotsEditorResults::GenerateTreeViewChildren_ModifiedActors(FLeve
 		
 		TreeViewModifiedActorGroupObjects.Add(NewActorGroup);
 	}
+
+	return TreeViewModifiedActorGroupObjects.Num() > 0;
 }
 
-void SLevelSnapshotsEditorResults::GenerateTreeViewChildren_AddedActors(FLevelSnapshotsEditorResultsRowPtr AddedActorsHeader)
+bool SLevelSnapshotsEditorResults::GenerateTreeViewChildren_AddedActors(FLevelSnapshotsEditorResultsRowPtr AddedActorsHeader)
 {
 	check(AddedActorsHeader);
 	
@@ -2048,9 +2164,11 @@ void SLevelSnapshotsEditorResults::GenerateTreeViewChildren_AddedActors(FLevelSn
 
 		TreeViewAddedActorGroupObjects.Add(NewActorRow);
 	}
+
+	return TreeViewAddedActorGroupObjects.Num() > 0;
 }
 
-void SLevelSnapshotsEditorResults::GenerateTreeViewChildren_RemovedActors(FLevelSnapshotsEditorResultsRowPtr RemovedActorsHeader)
+bool SLevelSnapshotsEditorResults::GenerateTreeViewChildren_RemovedActors(FLevelSnapshotsEditorResultsRowPtr RemovedActorsHeader)
 {
 	check(RemovedActorsHeader);
 	
@@ -2072,6 +2190,8 @@ void SLevelSnapshotsEditorResults::GenerateTreeViewChildren_RemovedActors(FLevel
 
 		TreeViewRemovedActorGroupObjects.Add(NewActorRow);
 	}
+
+	return TreeViewRemovedActorGroupObjects.Num() > 0;
 }
 
 FReply SLevelSnapshotsEditorResults::SetAllActorGroupsCollapsed()
@@ -2284,12 +2404,21 @@ void SLevelSnapshotsEditorResultsRow::Construct(const FArguments& InArgs, const
 		
 	// Splitter Slot 1
 	TSharedPtr<SWidget> SnapshotChildWidget;
-
+	
+	bool bIsSnapshotChildWidgetCustomized = false;
+	
 	if (const TSharedPtr<IPropertyHandle>& SnapshotPropertyHandle = PinnedItem->GetSnapshotPropertyHandle())
 	{
-		if (PinnedItem->GetWidgetTypeCustomization() == FLevelSnapshotsEditorResultsRow::WidgetType_PRGCustomWidget)
+		if (SnapshotPropertyHandle->IsCustomized())
 		{
-			SnapshotChildWidget = PinnedItem->GetSnapshotPrgWidget();
+			bIsSnapshotChildWidgetCustomized = true;
+			
+			if (const TSharedPtr<IDetailTreeNode>& Node = PinnedItem->GetSnapshotPropertyNode())
+			{
+				FNodeWidgets Widgets = Node->CreateNodeWidgets();
+
+				SnapshotChildWidget = Widgets.WholeRowWidget.IsValid() ? Widgets.WholeRowWidget : Widgets.ValueWidget;
+			}
 		}
 		else if (RowType == FLevelSnapshotsEditorResultsRow::SinglePropertyInMap)
 		{
@@ -2331,18 +2460,26 @@ void SLevelSnapshotsEditorResultsRow::Construct(const FArguments& InArgs, const
 	}
 
 	SnapshotChildWidget->SetEnabled(bIsHeaderRow);
+	SnapshotChildWidget->SetCanTick(bIsHeaderRow);
 
-	SplitterPtr->AddSlot()
-	[
-		SNew(SInvalidationPanel)
+	TSharedPtr<SWidget> FinalSnapshotWidget = SNew(SBox)
+		.VAlign(VAlign_Center)
+		.Padding(FMargin(2, 0))
 		[
-			SNew(SBox)
-			.VAlign(VAlign_Center)
-			.Padding(FMargin(2, 0))
+			SnapshotChildWidget.ToSharedRef()
+		];
+
+	if (!bIsSnapshotChildWidgetCustomized)
+	{
+		FinalSnapshotWidget = SNew(SInvalidationPanel)
 			[
-				SnapshotChildWidget.ToSharedRef()
-			]
-		]
+				FinalSnapshotWidget.ToSharedRef()
+			];
+	}
+
+	SplitterPtr->AddSlot()
+	[
+		FinalSnapshotWidget.ToSharedRef()
 	].OnSlotResized_Handler.BindSP(this, &SLevelSnapshotsEditorResultsRow::SetSnapshotColumnSize);
 
 	const auto SlotDelegate1 = TAttribute<float>::FGetter::CreateSP(this, &SLevelSnapshotsEditorResultsRow::GetSplitterSlotSize, 1);;
@@ -2351,11 +2488,20 @@ void SLevelSnapshotsEditorResultsRow::Construct(const FArguments& InArgs, const
 	// Splitter Slot 2
 	TSharedPtr<SWidget> WorldChildWidget;
 
+	bool bIsWorldChildWidgetCustomized = false;
+
 	if (const TSharedPtr<IPropertyHandle>& WorldPropertyHandle = PinnedItem->GetWorldPropertyHandle())
 	{
-		if (PinnedItem->GetWidgetTypeCustomization() == FLevelSnapshotsEditorResultsRow::WidgetType_PRGCustomWidget)
+		if (WorldPropertyHandle->IsCustomized())
 		{
-			WorldChildWidget = PinnedItem->GetWorldPrgWidget();
+			bIsWorldChildWidgetCustomized = true;
+			
+			if (const TSharedPtr<IDetailTreeNode>& Node = PinnedItem->GetWorldPropertyNode())
+			{
+				FNodeWidgets Widgets = Node->CreateNodeWidgets();
+
+				WorldChildWidget = Widgets.WholeRowWidget.IsValid() ? Widgets.WholeRowWidget : Widgets.ValueWidget;
+			}
 		}
 		else if (RowType == FLevelSnapshotsEditorResultsRow::SinglePropertyInMap)
 		{
@@ -2397,18 +2543,26 @@ void SLevelSnapshotsEditorResultsRow::Construct(const FArguments& InArgs, const
 	}
 
 	WorldChildWidget->SetEnabled(bIsHeaderRow);
+	WorldChildWidget->SetCanTick(bIsHeaderRow);
 
-	SplitterPtr->AddSlot()
-	[
-		SNew(SInvalidationPanel)
+	TSharedPtr<SWidget> FinalWorldWidget = SNew(SBox)
+		.VAlign(VAlign_Center)
+		.Padding(FMargin(2, 0))
 		[
-			SNew(SBox)
-			.VAlign(VAlign_Center)
-			.Padding(FMargin(2, 0))
+			WorldChildWidget.ToSharedRef()
+		];
+
+	if (!bIsWorldChildWidgetCustomized)
+	{
+		FinalWorldWidget = SNew(SInvalidationPanel)
 			[
-				WorldChildWidget.ToSharedRef()
-			]
-		]
+				FinalWorldWidget.ToSharedRef()
+			];
+	}
+
+	SplitterPtr->AddSlot()
+	[
+		FinalWorldWidget.ToSharedRef()
 	].OnSlotResized_Handler.BindSP(this, &SLevelSnapshotsEditorResultsRow::SetWorldObjectColumnSize);
 
 	const auto SlotDelegate2 = TAttribute<float>::FGetter::CreateSP(this, &SLevelSnapshotsEditorResultsRow::GetSplitterSlotSize, 2);;
diff --git a/Engine/Plugins/VirtualProduction/LevelSnapshots/Source/LevelSnapshotsEditor/Private/Views/Results/SLevelSnapshotsEditorResults.h b/Engine/Plugins/VirtualProduction/LevelSnapshots/Source/LevelSnapshotsEditor/Private/Views/Results/SLevelSnapshotsEditorResults.h
--- a/Engine/Plugins/VirtualProduction/LevelSnapshots/Source/LevelSnapshotsEditor/Private/Views/Results/SLevelSnapshotsEditorResults.h
+++ b/Engine/Plugins/VirtualProduction/LevelSnapshots/Source/LevelSnapshotsEditor/Private/Views/Results/SLevelSnapshotsEditorResults.h
@@ -2,33 +2,34 @@
 
 #pragma once
 
+#include "Widgets/Input/SCheckBox.h"
+#include "Widgets/SCompoundWidget.h"
+#include "Widgets/Views/STreeView.h"
+
 #include "FilterListData.h"
 #include "Data/LevelSnapshotsEditorData.h"
 
-#include "Widgets/DeclarativeSyntaxSupport.h"
-#include "Widgets/Layout/SScrollBox.h"
-#include "Widgets/SCompoundWidget.h"
-#include "Widgets/Views/STreeView.h"
+class SLevelSnapshotsEditorResults;
+class ULevelSnapshot;
 
+class FMenuBuilder;
+class FPropertyEditorModule;
+class FReply;
+class IDetailTreeNode;
 class IPropertyHandle;
 class IPropertyRowGenerator;
-class FPropertyEditorModule;
-class SLevelSnapshotsEditorResults;
-class FLevelSnapshotsEditorResults;
 class SBox;
-class SButton;
-class SSearchBox;
 class SSplitter;
-class UFilteredResults;
+class SSearchBox;
+class STextBlock;
+
+class ULevelSnapshotsEditorData;
+
+struct EVisibility;
+struct FPropertyRowGeneratorArgs;
 
-struct FFilterListData;
-struct FLevelSnapshotPropertyChain;
 struct FLevelSnapshotsEditorResultsRow;
 struct FLevelSnapshotsEditorResultsSplitterManager;
-struct FLevelSnapshotsEditorViewBuilder;
-struct FPropertyRowGeneratorArgs;
-struct FPropertySelection;
-struct FPropertySelectionMap;
 
 typedef TSharedPtr<FLevelSnapshotsEditorResultsRow> FLevelSnapshotsEditorResultsRowPtr;
 typedef TSharedPtr<FLevelSnapshotsEditorResultsSplitterManager> FLevelSnapshotsEditorResultsSplitterManagerPtr;
@@ -61,15 +62,9 @@ struct FRowGeneratorInfo
 		, GeneratorObject(InGeneratorObject)
 	{};
 	
-	TWeakPtr<IPropertyRowGenerator> GetGeneratorObject() const
-	{
-		return GeneratorObject;
-	}
+	TWeakPtr<IPropertyRowGenerator> GetGeneratorObject() const;
 
-	void FlushReferences()
-	{
-		GeneratorObject.Reset();
-	}
+	void FlushReferences();
 
 private:
 	/* The object that represents the object passed into the generator */
@@ -80,6 +75,19 @@ private:
 	TSharedPtr<IPropertyRowGenerator> GeneratorObject;
 };
 
+struct FPropertyHandleHierarchy
+{
+	FPropertyHandleHierarchy(const TSharedPtr<IDetailTreeNode>& InNode, const TSharedPtr<IPropertyHandle>& InHandle);
+
+	// Used to identify counterparts
+	FString PropertyName;
+	FString ParentPropertyName;
+
+	TSharedPtr<IDetailTreeNode> Node;
+	TSharedPtr<IPropertyHandle> Handle;
+	TArray<TSharedRef<FPropertyHandleHierarchy>> DirectChildren;
+};
+
 struct FLevelSnapshotsEditorResultsRow final : TSharedFromThis<FLevelSnapshotsEditorResultsRow>
 {
 	enum ELevelSnapshotsEditorResultsRowType
@@ -107,12 +115,6 @@ struct FLevelSnapshotsEditorResultsRow final : TSharedFromThis<FLevelSnapshotsEd
 		HeaderType_RemovedActors,
 	};
 
-	enum ELevelSnapshotsWidgetTypeCustomization
-	{
-		WidgetType_NoCustomWidget,
-		WidgetType_PRGCustomWidget
-	};
-
 	~FLevelSnapshotsEditorResultsRow();
 
 	void FlushReferences();
@@ -128,20 +130,15 @@ struct FLevelSnapshotsEditorResultsRow final : TSharedFromThis<FLevelSnapshotsEd
 	
 	void InitActorRow(AActor* InSnapshotActor, AActor* InWorldActor, const TWeakPtr<SLevelSnapshotsEditorResults>& InResultsView);
 	void InitObjectRow(
+		UObject* InSnapshotObject, UObject* InWorldObject,
 		const TWeakPtr<FRowGeneratorInfo>& InSnapshotRowGenerator,
 		const TWeakPtr<FRowGeneratorInfo>& InWorldRowGenerator,
 		const TWeakPtr<SLevelSnapshotsEditorResults>& InResultsView);
 
-	void InitPropertyRowWithHandles(
+	void InitPropertyRow(
 		const TWeakPtr<FLevelSnapshotsEditorResultsRow>& InContainingObjectGroup,
-		const TSharedPtr<IPropertyHandle>& InSnapshotHandle, const TSharedPtr<IPropertyHandle>& InWorldHandle, 
-		const bool bNewIsCounterpartValueSame = false, const ELevelSnapshotsWidgetTypeCustomization InWidgetTypeCustomization = WidgetType_NoCustomWidget);
-
-	void InitPropertyRowWithCustomizedWidgets(
-		const TWeakPtr<FLevelSnapshotsEditorResultsRow>& InContainingObjectGroup,
-		const TSharedPtr<IPropertyHandle>& InSnapshotHandle, const TSharedPtr<IPropertyHandle>& InWorldHandle,
-		const TSharedPtr<SWidget>& InSnapshotWidget, const TSharedPtr<SWidget>& InWorldWidget,
-		const bool bNewIsCounterpartValueSame = false, const ELevelSnapshotsWidgetTypeCustomization InWidgetTypeCustomization = WidgetType_PRGCustomWidget);
+		const TSharedPtr<FPropertyHandleHierarchy>& InSnapshotHierarchy, const TSharedPtr<FPropertyHandleHierarchy>& InWorldHandleHierarchy,
+		const bool bNewIsCounterpartValueSame = false);
 
 	void GenerateActorGroupChildren(FPropertySelectionMap& PropertySelectionMap);
 
@@ -177,10 +174,6 @@ struct FLevelSnapshotsEditorResultsRow final : TSharedFromThis<FLevelSnapshotsEd
 
 	TWeakPtr<FLevelSnapshotsEditorResultsRow> GetContainingObjectGroup() const;
 
-	ELevelSnapshotsWidgetTypeCustomization GetWidgetTypeCustomization() const;
-	const TSharedPtr<SWidget>& GetWorldPrgWidget() const;
-	const TSharedPtr<SWidget>& GetSnapshotPrgWidget() const;
-
 	bool GetHasGeneratedChildren() const;
 	void SetHasGeneratedChildren(const bool bNewGenerated);
 
@@ -201,8 +194,12 @@ struct FLevelSnapshotsEditorResultsRow final : TSharedFromThis<FLevelSnapshotsEd
 	FProperty* GetProperty() const;
 	FLevelSnapshotPropertyChain GetPropertyChain() const;
 
-	const TSharedPtr<IPropertyHandle>& GetSnapshotPropertyHandle() const;
-	const TSharedPtr<IPropertyHandle>& GetWorldPropertyHandle() const;
+	TSharedPtr<IDetailTreeNode> GetSnapshotPropertyNode() const;
+	TSharedPtr<IDetailTreeNode> GetWorldPropertyNode() const;
+	ELevelSnapshotsObjectType GetFirstValidPropertyNode(TSharedPtr<IDetailTreeNode>& OutNode) const;
+
+	TSharedPtr<IPropertyHandle> GetSnapshotPropertyHandle() const;
+	TSharedPtr<IPropertyHandle> GetWorldPropertyHandle() const;
 	ELevelSnapshotsObjectType GetFirstValidPropertyHandle(TSharedPtr<IPropertyHandle>& OutHandle) const;
 
 	bool GetIsCounterpartValueSame() const;
@@ -224,8 +221,8 @@ struct FLevelSnapshotsEditorResultsRow final : TSharedFromThis<FLevelSnapshotsEd
 	/* Whether the group has any children with associated properties that have any difference between the chosen snapshot and the current level. */
 	bool HasChangedChildren() const;
 
-	/* Gets child single properties, whole collection properties and individual properties inside of structs that belong to a specific group. Skips nested object groups like components and subobjects. */
 	void GetAllCheckedChildProperties(TArray<FLevelSnapshotsEditorResultsRowPtr>& CheckedSinglePropertyNodeArray) const;
+	void GetAllUncheckedChildProperties(TArray<FLevelSnapshotsEditorResultsRowPtr>& UncheckedSinglePropertyNodeArray) const;
 	
 	EVisibility GetDesiredVisibility() const;
 
@@ -248,11 +245,6 @@ private:
 	/* This is the component, subobject or actor group to which this row belongs. If nullptr, this row is a top-level actor group. */
 	TWeakPtr<FLevelSnapshotsEditorResultsRow> ContainingObjectGroup = nullptr;
 
-	// Whether we will generate a custom widget for this row instead of using a handle from PropertyRowGenerator
-	ELevelSnapshotsWidgetTypeCustomization WidgetTypeCustomization = WidgetType_NoCustomWidget;
-	TSharedPtr<SWidget> PRG_WorldCustomWidget;
-	TSharedPtr<SWidget> PRG_SnapshotCustomWidget;
-
 	// Only applies to object groups - all of the property groups are generated with the rest of the single properties
 	bool bHasGeneratedChildren = false;
 
@@ -281,8 +273,8 @@ private:
 
 	/* For property type rows */
 
-	TSharedPtr<IPropertyHandle> SnapshotPropertyHandle;
-	TSharedPtr<IPropertyHandle> WorldPropertyHandle;
+	TSharedPtr<FPropertyHandleHierarchy> SnapshotPropertyHandleHierarchy;
+	TSharedPtr<FPropertyHandleHierarchy> WorldPropertyHandleHierarchy;
 
 	/* Whether the snapshot and world object properties have the same value */
 	bool bIsCounterpartValueSame = false;
@@ -306,16 +298,18 @@ public:
 
 	SLATE_END_ARGS()
 
-	enum ERefreshStyle
-	{
-		FiltersChanged,
-		FiltersNotChanged
-	};
-
 	void Construct(const FArguments& InArgs, ULevelSnapshotsEditorData* InEditorData);
 
 	virtual ~SLevelSnapshotsEditorResults() override;
 
+	FMenuBuilder BuildShowOptionsMenu();
+
+	void SetShowFilteredActors(const bool bNewSetting);
+	void SetShowUnselectedActors(const bool bNewSetting);
+
+	bool GetShowFilteredActors() const;
+	bool GetShowUnselectedActors() const;
+	
 	void FlushMemory();
 
 	TOptional<ULevelSnapshot*> GetSelectedLevelSnapshot() const;
@@ -372,15 +366,20 @@ private:
 	TSharedPtr<SSearchBox> ResultsSearchBoxPtr;
 	TSharedPtr<SBox> ResultsBoxContainerPtr;
 
+	// 'Show' Options
+
+	bool bShowFilteredActors = false;
+	bool bShowUnselectedActors = true;
+
 	/* For splitter sync */
 	FLevelSnapshotsEditorResultsSplitterManagerPtr SplitterManagerPtr;
 
 	//  Tree View Implementation
 
 	void GenerateTreeView();
-	void GenerateTreeViewChildren_ModifiedActors(FLevelSnapshotsEditorResultsRowPtr ModifiedActorsHeader, ULevelSnapshotFilter* UserFilters);
-	void GenerateTreeViewChildren_AddedActors(FLevelSnapshotsEditorResultsRowPtr AddedActorsHeader);
-	void GenerateTreeViewChildren_RemovedActors(FLevelSnapshotsEditorResultsRowPtr RemovedActorsHeader);
+	bool GenerateTreeViewChildren_ModifiedActors(FLevelSnapshotsEditorResultsRowPtr ModifiedActorsHeader, ULevelSnapshotFilter* UserFilters);
+	bool GenerateTreeViewChildren_AddedActors(FLevelSnapshotsEditorResultsRowPtr AddedActorsHeader);
+	bool GenerateTreeViewChildren_RemovedActors(FLevelSnapshotsEditorResultsRowPtr RemovedActorsHeader);
 	
 	void OnGetRowChildren(FLevelSnapshotsEditorResultsRowPtr Row, TArray<FLevelSnapshotsEditorResultsRowPtr>& OutChildren);
 	void OnRowChildExpansionChange(FLevelSnapshotsEditorResultsRowPtr Row, const bool bIsExpanded);
diff --git a/Engine/Source/Editor/PropertyEditor/Private/PropertyRowGenerator.cpp b/Engine/Source/Editor/PropertyEditor/Private/PropertyRowGenerator.cpp
--- a/Engine/Source/Editor/PropertyEditor/Private/PropertyRowGenerator.cpp
+++ b/Engine/Source/Editor/PropertyEditor/Private/PropertyRowGenerator.cpp
@@ -643,8 +643,13 @@ void FPropertyRowGenerator::UpdateSinglePropertyMap(TSharedPtr<FComplexPropertyN
 }
 
 
-bool FPropertyRowGenerator::ValidatePropertyNodes(const FRootPropertyNodeList &PropertyNodeList)
+bool FPropertyRowGenerator::ValidatePropertyNodes(const FRootPropertyNodeList& PropertyNodeList)
 {
+	if (CustomValidatePropertyNodesFunction.IsBound())
+	{
+		return CustomValidatePropertyNodesFunction.Execute(PropertyNodeList);
+	}
+	
 	bool bFullRefresh = false;
 
 	for (const TSharedPtr<FComplexPropertyNode>& RootPropertyNode : RootPropertyNodes)
diff --git a/Engine/Source/Editor/PropertyEditor/Private/PropertyRowGenerator.h b/Engine/Source/Editor/PropertyEditor/Private/PropertyRowGenerator.h
--- a/Engine/Source/Editor/PropertyEditor/Private/PropertyRowGenerator.h
+++ b/Engine/Source/Editor/PropertyEditor/Private/PropertyRowGenerator.h
@@ -69,8 +69,6 @@ struct FDetailLayoutData
 
 typedef TArray<FDetailLayoutData> FDetailLayoutList;
 
-typedef TArray<TSharedPtr<FComplexPropertyNode>> FRootPropertyNodeList;
-
 class FPropertyRowGenerator : public IPropertyRowGenerator, public FTickableEditorObject, public TSharedFromThis<FPropertyRowGenerator>
 {
 public:
@@ -95,6 +93,14 @@ public:
 	virtual void RegisterInstancedCustomPropertyTypeLayout(FName PropertyTypeName, FOnGetPropertyTypeCustomizationInstance PropertyTypeLayoutDelegate, TSharedPtr<IPropertyTypeIdentifier> Identifier = nullptr) override;
 	virtual void UnregisterInstancedCustomPropertyLayout(UStruct* Class) override;
 	virtual void UnregisterInstancedCustomPropertyTypeLayout(FName PropertyTypeName, TSharedPtr<IPropertyTypeIdentifier> Identifier = nullptr) override;
+	virtual TSharedPtr<FAssetThumbnailPool> GetGeneratedThumbnailPool() override
+	{ 
+		return GetThumbnailPool();
+	}
+	virtual void SetCustomValidatePropertyNodesFunction(FOnValidatePropertyRowGeneratorNodes InCustomValidatePropertyNodesFunction) override
+	{
+		CustomValidatePropertyNodesFunction = MoveTemp(InCustomValidatePropertyNodesFunction);
+	}
 
 	/** FTickableEditorObject interface */
 	virtual ETickableTickType GetTickableTickType() const override { return ETickableTickType::Always; }
@@ -150,6 +156,8 @@ private:
 	TSharedRef<IPropertyGenerationUtilities> PropertyGenerationUtilities;
 	/** Delegate called when the details panel finishes editing a property (after post edit change is called) */
 	FOnFinishedChangingProperties OnFinishedChangingPropertiesDelegate;
+	/** The ValidatePropertyNodes function can be overridden with this member, if set. Useful if your implementation doesn't require this kind of validation each Tick. */
+	FOnValidatePropertyRowGeneratorNodes CustomValidatePropertyNodesFunction;
 
 	bool bViewingClassDefaultObject;
 };
diff --git a/Engine/Source/Editor/PropertyEditor/Public/IPropertyRowGenerator.h b/Engine/Source/Editor/PropertyEditor/Public/IPropertyRowGenerator.h
--- a/Engine/Source/Editor/PropertyEditor/Public/IPropertyRowGenerator.h
+++ b/Engine/Source/Editor/PropertyEditor/Public/IPropertyRowGenerator.h
@@ -2,11 +2,17 @@
 
 #pragma once
 
+#include "CoreMinimal.h"
 #include "PropertyEditorModule.h"
 
 class IDetailTreeNode;
+class FComplexPropertyNode;
 class FStructOnScope;
 
+typedef TArray<TSharedPtr<FComplexPropertyNode>> FRootPropertyNodeList;
+
+DECLARE_DELEGATE_RetVal_OneParam(bool, FOnValidatePropertyRowGeneratorNodes, const FRootPropertyNodeList&)
+
 struct FPropertyRowGeneratorArgs
 {
 	FPropertyRowGeneratorArgs()
@@ -101,4 +107,8 @@ public:
 	virtual void UnregisterInstancedCustomPropertyTypeLayout(FName PropertyTypeName, TSharedPtr<IPropertyTypeIdentifier> Identifier = nullptr) = 0;
 
 	virtual FOnFinishedChangingProperties& OnFinishedChangingProperties() = 0;
+
+	/* Use this function to set a callback on FPropertyRowGenerator that will override the ValidatePropertyNodes function.
+	 * This is useful if your implementation doesn't need to validate nodes every tick or needs to perform some other form of validation. */
+	virtual void SetCustomValidatePropertyNodesFunction(FOnValidatePropertyRowGeneratorNodes InCustomValidatePropertyNodesFunction) = 0;
 };
\ No newline at end of file
