Fix Horde Issue 71820

diff --git a/Engine/Plugins/VirtualProduction/LevelSnapshots/Source/LevelSnapshots/Private/Archive/ApplySnapshotDataArchiveV2.cpp b/Engine/Plugins/VirtualProduction/LevelSnapshots/Source/LevelSnapshots/Private/Archive/ApplySnapshotDataArchiveV2.cpp
--- a/Engine/Plugins/VirtualProduction/LevelSnapshots/Source/LevelSnapshots/Private/Archive/ApplySnapshotDataArchiveV2.cpp
+++ b/Engine/Plugins/VirtualProduction/LevelSnapshots/Source/LevelSnapshots/Private/Archive/ApplySnapshotDataArchiveV2.cpp
@@ -61,8 +61,7 @@ namespace
 			:
 			Super(SaveLocation),
 			PropertiesToSerialize(PropertiesToSerialize),
-			SnapshotObject(SnapshotObject),
-			TextProperties(TextProperties)
+			SnapshotObject(SnapshotObject)
 		{
 			ArNoDelta = true;
 		}
diff --git a/Engine/Plugins/VirtualProduction/LevelSnapshots/Source/LevelSnapshots/Private/Archive/SnapshotArchive.h b/Engine/Plugins/VirtualProduction/LevelSnapshots/Source/LevelSnapshots/Private/Archive/SnapshotArchive.h
--- a/Engine/Plugins/VirtualProduction/LevelSnapshots/Source/LevelSnapshots/Private/Archive/SnapshotArchive.h
+++ b/Engine/Plugins/VirtualProduction/LevelSnapshots/Source/LevelSnapshots/Private/Archive/SnapshotArchive.h
@@ -4,6 +4,7 @@
 
 #include "CoreMinimal.h"
 #include "Serialization/ArchiveUObject.h"
+#include "UObject/ObjectMacros.h"
 
 class UObject;
 class UPackage;
@@ -41,7 +42,7 @@ protected:
 		return SharedData;
 	}
 
-	int32 ExcludedPropertyFlags;
+	EPropertyFlags ExcludedPropertyFlags;
 
 	FSnapshotArchive(FObjectSnapshotData& InObjectData, FWorldSnapshotData& InSharedData, bool bIsLoading);
 	
