Fix Horde Issue 71830

diff --git a/Engine/Plugins/VirtualProduction/LevelSnapshots/Source/LevelSnapshots/Private/Archive/SnapshotArchive.cpp b/Engine/Plugins/VirtualProduction/LevelSnapshots/Source/LevelSnapshots/Private/Archive/SnapshotArchive.cpp
--- a/Engine/Plugins/VirtualProduction/LevelSnapshots/Source/LevelSnapshots/Private/Archive/SnapshotArchive.cpp
+++ b/Engine/Plugins/VirtualProduction/LevelSnapshots/Source/LevelSnapshots/Private/Archive/SnapshotArchive.cpp
@@ -12,7 +12,9 @@
 void FSnapshotArchive::RestoreData(FObjectSnapshotData& InObjectData, FWorldSnapshotData& InSharedData, UObject* InObjectToRestore, UPackage* InLocalisationSnapshotPackage)
 {
 	FSnapshotArchive Archive(InObjectData, InSharedData, true);
+#if USE_STABLE_LOCALIZATION_KEYS
 	Archive.SetLocalizationNamespace(TextNamespaceUtil::EnsurePackageNamespace(InLocalisationSnapshotPackage));
+#endif
 	InObjectToRestore->Serialize(Archive);
 }
 
