Fix filtering for added and removed actors not working

 Fix certain Blueprint filters do not show up

Fix crash when removing filter classes

Avoid leaking class defaults

Move details customisation from global to local

Changed function name to be more self-expressive

#[fyi] Jason.Walter
Fix snapshots leaking when created via Blueprints (removed Standalone flag)
Emit warnings when performing unsupported snapshot operation (taking and applying)

Handle migrating assets
Fix bug where object references would only be resolved but not loaded

Fix Search issues:
-Added descriptive comments to ELevelSnapshotsEditorResultsRowType members which needed more clarification
-Added method to be able to set a row's search terms directly from outside of the row's scope. This is useful to add all properties that pass filters to a string and cache it on an actor group that hasn't been expanded yet. This allows the group to be correctly searched when its children have not yet been generated.
-Added method to execute search on specific actors only. When a group has its children generated this method allows us to apply search to the newly expanded group alone instead of all rows.
-Fixed an issue where parent groups of a row that matches search were not shown. Now the whole parentage of a matching row will be shown when searching.
-Fixed an issue where refreshing results or applying to world would not respect the last search.
-Fixed an issue where searching then expanding a group for the first time would not respect the last search.
-Fixed an issue where searching did not work when an item had not yet generated its children.
-Fixed an issue in which recursive selection and deselection (where the child groups are set to the same checked state as their parent) did not work on items that were invisible due to active search.

Misc:
-Added help text to inform the user that there are no results to display along with some tips should they feel they should be seeing results.

#[fyi] Jason.Walter

diff --git a/Engine/Plugins/VirtualProduction/LevelSnapshots/Source/LevelSnapshots/Private/Archive/ApplySnapshotDataArchiveV2.h b/Engine/Plugins/VirtualProduction/LevelSnapshots/Source/LevelSnapshots/Private/Archive/ApplySnapshotDataArchiveV2.h
--- a/Engine/Plugins/VirtualProduction/LevelSnapshots/Source/LevelSnapshots/Private/Archive/ApplySnapshotDataArchiveV2.h
+++ b/Engine/Plugins/VirtualProduction/LevelSnapshots/Source/LevelSnapshots/Private/Archive/ApplySnapshotDataArchiveV2.h
@@ -7,7 +7,7 @@
 #include "PropertySelectionMap.h"
 #include "SnapshotArchive.h"
 
-/* For writing data into an object */
+/* For writing data into an editor object. */
 class FApplySnapshotDataArchiveV2 : public FSnapshotArchive
 {
 	using Super = FSnapshotArchive;
diff --git a/Engine/Plugins/VirtualProduction/LevelSnapshots/Source/LevelSnapshots/Private/Archive/SnapshotArchive.cpp b/Engine/Plugins/VirtualProduction/LevelSnapshots/Source/LevelSnapshots/Private/Archive/SnapshotArchive.cpp
--- a/Engine/Plugins/VirtualProduction/LevelSnapshots/Source/LevelSnapshots/Private/Archive/SnapshotArchive.cpp
+++ b/Engine/Plugins/VirtualProduction/LevelSnapshots/Source/LevelSnapshots/Private/Archive/SnapshotArchive.cpp
@@ -3,13 +3,14 @@
 #include "SnapshotArchive.h"
 
 #include "ObjectSnapshotData.h"
+#include "SnapshotVersion.h"
 #include "WorldSnapshotData.h"
 
 #include "Internationalization/TextNamespaceUtil.h"
 #include "Internationalization/TextPackageNamespaceUtil.h"
 #include "UObject/ObjectMacros.h"
 
-void FSnapshotArchive::RestoreData(FObjectSnapshotData& InObjectData, FWorldSnapshotData& InSharedData, UObject* InObjectToRestore, UPackage* InLocalisationSnapshotPackage)
+void FSnapshotArchive::ApplyToSnapshotWorldObject(FObjectSnapshotData& InObjectData, FWorldSnapshotData& InSharedData, UObject* InObjectToRestore, UPackage* InLocalisationSnapshotPackage)
 {
 	FSnapshotArchive Archive(InObjectData, InSharedData, true);
 #if USE_STABLE_LOCALIZATION_KEYS
@@ -160,4 +161,20 @@ FSnapshotArchive::FSnapshotArchive(FObjectSnapshotData& InObjectData, FWorldSnap
 		Super::SetIsLoading(false);
 		Super::SetIsSaving(true);
 	}
+
+	if (bIsLoading)
+	{
+		const FSnapshotVersionInfo& VersionInfo = InSharedData.GetSnapshotVersionInfo();
+		
+		Super::SetUE4Ver(VersionInfo.FileVersion.FileVersionUE4);
+		Super::SetLicenseeUE4Ver(VersionInfo.FileVersion.FileVersionLicenseeUE4);
+		Super::SetEngineVer(FEngineVersionBase(VersionInfo.EngineVersion.Major, VersionInfo.EngineVersion.Minor, VersionInfo.EngineVersion.Patch, VersionInfo.EngineVersion.Changelist));
+
+		FCustomVersionContainer EngineCustomVersions;
+		for (const FSnapshotCustomVersionInfo& CustomVersion : VersionInfo.CustomVersions)
+		{
+			EngineCustomVersions.SetVersion(CustomVersion.Key, CustomVersion.Version, CustomVersion.FriendlyName);
+		}
+		Super::SetCustomVersions(EngineCustomVersions);
+	}
 }
diff --git a/Engine/Plugins/VirtualProduction/LevelSnapshots/Source/LevelSnapshots/Private/Archive/SnapshotArchive.h b/Engine/Plugins/VirtualProduction/LevelSnapshots/Source/LevelSnapshots/Private/Archive/SnapshotArchive.h
--- a/Engine/Plugins/VirtualProduction/LevelSnapshots/Source/LevelSnapshots/Private/Archive/SnapshotArchive.h
+++ b/Engine/Plugins/VirtualProduction/LevelSnapshots/Source/LevelSnapshots/Private/Archive/SnapshotArchive.h
@@ -3,6 +3,7 @@
 #pragma once
 
 #include "CoreMinimal.h"
+#include "SnapshotVersion.h"
 #include "Serialization/ArchiveUObject.h"
 #include "UObject/ObjectMacros.h"
 
@@ -18,7 +19,7 @@ class FSnapshotArchive : public FArchiveUObject
 	
 public:
 
-	static void RestoreData(FObjectSnapshotData& InObjectData, FWorldSnapshotData& InSharedData, UObject* InObjectToRestore, UPackage* InLocalisationSnapshotPackage);
+	static void ApplyToSnapshotWorldObject(FObjectSnapshotData& InObjectData, FWorldSnapshotData& InSharedData, UObject* InObjectToRestore, UPackage* InLocalisationSnapshotPackage);
 
 	//~ Begin FArchive Interface
 	virtual FString GetArchiveName() const override;
@@ -44,6 +45,11 @@ protected:
 
 	EPropertyFlags ExcludedPropertyFlags;
 
+	/**
+	 * @param InObjectData Holds the array we're loading from or saving to
+	 * @param InSharedData Used to store shared data, e.g. references
+	 * @param bIsLoading Whether to load or save
+	 */
 	FSnapshotArchive(FObjectSnapshotData& InObjectData, FWorldSnapshotData& InSharedData, bool bIsLoading);
 	
 private:
diff --git a/Engine/Plugins/VirtualProduction/LevelSnapshots/Source/LevelSnapshots/Private/Data/ActorSnapshotData.cpp b/Engine/Plugins/VirtualProduction/LevelSnapshots/Source/LevelSnapshots/Private/Data/ActorSnapshotData.cpp
--- a/Engine/Plugins/VirtualProduction/LevelSnapshots/Source/LevelSnapshots/Private/Data/ActorSnapshotData.cpp
+++ b/Engine/Plugins/VirtualProduction/LevelSnapshots/Source/LevelSnapshots/Private/Data/ActorSnapshotData.cpp
@@ -142,7 +142,7 @@ TOptional<AActor*> FActorSnapshotData::GetDeserialized(UWorld* SnapshotWorld, FW
 	bReceivedSerialisation = true;
 	
 	AActor* PreallocatedActor = Preallocated.GetValue();
-	FSnapshotArchive::RestoreData(SerializedActorData, WorldData, PreallocatedActor, InLocalisationSnapshotPackage);
+	FSnapshotArchive::ApplyToSnapshotWorldObject(SerializedActorData, WorldData, PreallocatedActor, InLocalisationSnapshotPackage);
 
 	DeserializeComponents(PreallocatedActor, WorldData, [InLocalisationSnapshotPackage](FObjectSnapshotData& SerializedCompData, FComponentSnapshotData& CompData, UActorComponent* Comp, FWorldSnapshotData& SharedData)
 	{
diff --git a/Engine/Plugins/VirtualProduction/LevelSnapshots/Source/LevelSnapshots/Private/Data/ComponentSnapshotData.cpp b/Engine/Plugins/VirtualProduction/LevelSnapshots/Source/LevelSnapshots/Private/Data/ComponentSnapshotData.cpp
--- a/Engine/Plugins/VirtualProduction/LevelSnapshots/Source/LevelSnapshots/Private/Data/ComponentSnapshotData.cpp
+++ b/Engine/Plugins/VirtualProduction/LevelSnapshots/Source/LevelSnapshots/Private/Data/ComponentSnapshotData.cpp
@@ -58,5 +58,5 @@ void FComponentSnapshotData::DeserializeIntoTransient(FObjectSnapshotData& Seria
 		//   - We apply the CDO and afterwards we override it with the serialized data. Good.
 	WorldData.SerializeClassDefaultsInto(ComponentToDeserializeInto);
 	
-	FSnapshotArchive::RestoreData(SerializedComponentData, WorldData, ComponentToDeserializeInto, InLocalisationSnapshotPackage);
+	FSnapshotArchive::ApplyToSnapshotWorldObject(SerializedComponentData, WorldData, ComponentToDeserializeInto, InLocalisationSnapshotPackage);
 }
diff --git a/Engine/Plugins/VirtualProduction/LevelSnapshots/Source/LevelSnapshots/Private/Data/LevelSnapshot.cpp b/Engine/Plugins/VirtualProduction/LevelSnapshots/Source/LevelSnapshots/Private/Data/LevelSnapshot.cpp
--- a/Engine/Plugins/VirtualProduction/LevelSnapshots/Source/LevelSnapshots/Private/Data/LevelSnapshot.cpp
+++ b/Engine/Plugins/VirtualProduction/LevelSnapshots/Source/LevelSnapshots/Private/Data/LevelSnapshot.cpp
@@ -16,8 +16,8 @@
 #include "UObject/Package.h"
 #include "UObject/TextProperty.h"
 #include "UObject/UnrealType.h"
-
 #if WITH_EDITOR
+#include "Logging/MessageLog.h"
 #include "ScopedTransaction.h"
 #endif
 
@@ -80,8 +80,29 @@ namespace
 
 void ULevelSnapshot::ApplySnapshotToWorld(UWorld* TargetWorld, const FPropertySelectionMap& SelectionSet)
 {
+	if (TargetWorld == nullptr)
+	{
+		return;
+	}
+	
+	if (MapPath != FSoftObjectPath(TargetWorld))
+	{
+#if !(UE_BUILD_SHIPPING || UE_BUILD_TEST)
+		if (TargetWorld->IsPlayInEditor())
+		{
+			FMessageLog("PIE").Warning(
+				FText::Format(NSLOCTEXT("LevelSnapshots", "IncompatibleWorlds", "This snapshot was taken in world '%s' and cannot be applied to PIE world '%s'. Snapshots can only be applied to the world they were taken in."),
+				FText::FromString(MapPath.ToString()),
+				FText::FromString(TargetWorld->GetPathName())
+				)
+			);
+		}
+#endif
+		UE_LOG(LogLevelSnapshots, Error, TEXT("This snapshot was taken for world '%s' and cannot be applied to world '%s': snapshots currently only support applying to the world they were taken in. "), *MapPath.ToString(), *TargetWorld->GetPathName());
+		return;
+	}
+	
 	EnsureWorldInitialised();
-
 #if WITH_EDITOR
 	FScopedTransaction Transaction(FText::FromString("Loading Level Snapshot."));
 #endif
@@ -107,6 +128,19 @@ void ULevelSnapshot::SnapshotWorld(UWorld* TargetWorld)
 		return;
 	}
 
+	if (TargetWorld->WorldType != EWorldType::Editor)
+	{
+#if !(UE_BUILD_SHIPPING || UE_BUILD_TEST)
+		if (TargetWorld->IsPlayInEditor())
+		{
+			FMessageLog("PIE").Warning(
+			NSLOCTEXT("LevelSnapshots", "IncompatibleWorlds", "Taking snapshots in PIE is an experimental feature. The snapshot will work in the same PIE session but may no longer work when you start a new PIE session.")
+			);
+		}
+#endif
+		UE_LOG(LogLevelSnapshots, Warning, TEXT("Level snapshots currently only support editors. Snapshots taken in other world types are experimental any may not function as expected."));
+	}
+	
 	const bool bHasLegacyData = ActorSnapshots.Num() > 0;
 	if (bHasLegacyData)
 	{
diff --git a/Engine/Plugins/VirtualProduction/LevelSnapshots/Source/LevelSnapshots/Private/Data/SnapshotVersion.cpp b/Engine/Plugins/VirtualProduction/LevelSnapshots/Source/LevelSnapshots/Private/Data/SnapshotVersion.cpp
new file mode 100644
--- /dev/null
+++ b/Engine/Plugins/VirtualProduction/LevelSnapshots/Source/LevelSnapshots/Private/Data/SnapshotVersion.cpp
@@ -0,0 +1,44 @@
+// Copyright Epic Games, Inc. All Rights Reserved.
+
+#include "Data/SnapshotVersion.h"
+#include "Misc/EngineVersion.h"
+#include "Serialization/CustomVersion.h"
+void FSnapshotFileVersionInfo::Initialize()
+{
+	FileVersionUE4 = GPackageFileUE4Version;
+	FileVersionLicenseeUE4 = GPackageFileLicenseeUE4Version;
+}
+
+void FSnapshotEngineVersionInfo::Initialize(const FEngineVersion& InVersion)
+{
+	Major = InVersion.GetMajor();
+	Minor = InVersion.GetMinor();
+	Patch = InVersion.GetPatch();
+	Changelist = InVersion.GetChangelist();
+}
+
+void FSnapshotCustomVersionInfo::Initialize(const FCustomVersion& InVersion)
+{
+	FriendlyName = InVersion.GetFriendlyName();
+	Key = InVersion.Key;
+	Version = InVersion.Version;
+}
+
+void FSnapshotVersionInfo::Initialize()
+{
+	FileVersion.Initialize();
+	EngineVersion.Initialize(FEngineVersion::Current());
+
+	CustomVersions.Empty(CustomVersions.Num());
+	FCustomVersionContainer AllCurrentVersions = FCurrentCustomVersions::GetAll();
+	for (const FCustomVersion& EngineCustomVersion : AllCurrentVersions.GetAllVersions())
+	{
+		FSnapshotCustomVersionInfo& CustomVersion = CustomVersions.AddDefaulted_GetRef();
+		CustomVersion.Initialize(EngineCustomVersion);
+	}
+}
+
+bool FSnapshotVersionInfo::IsInitialized() const
+{
+	return CustomVersions.Num() > 0;
+}
diff --git a/Engine/Plugins/VirtualProduction/LevelSnapshots/Source/LevelSnapshots/Private/Data/WorldSnapshotData.cpp b/Engine/Plugins/VirtualProduction/LevelSnapshots/Source/LevelSnapshots/Private/Data/WorldSnapshotData.cpp
--- a/Engine/Plugins/VirtualProduction/LevelSnapshots/Source/LevelSnapshots/Private/Data/WorldSnapshotData.cpp
+++ b/Engine/Plugins/VirtualProduction/LevelSnapshots/Source/LevelSnapshots/Private/Data/WorldSnapshotData.cpp
@@ -169,6 +169,12 @@
 void FWorldSnapshotData::OnDestroySnapshotWorld()
 {
 	TempActorWorld.Reset();
+
+	// Avoid object leaking
+	for (auto ClassDefaultsIt = ClassDefaults.CreateIterator(); ClassDefaultsIt; ++ClassDefaultsIt)
+	{
+		ClassDefaultsIt->Value.CachedLoadedClassDefault = nullptr;
+	}
 }
 
 void FWorldSnapshotData::SnapshotWorld(UWorld* World)
@@ -180,12 +186,14 @@
 	SerializedNames.Empty();
 	SerializedObjectReferences.Empty();
 	Subobjects.Empty();
+	
+	SnapshotVersionInfo.Initialize();
 
 #if WITH_EDITOR
 	FScopedSlowTask TakeSnapshotTask(World->GetProgressDenominator(), LOCTEXT("TakeSnapshotKey", "Take snapshot"));
 	TakeSnapshotTask.MakeDialogDelayed(1.f);
 #endif
-	
+
 	for (TActorIterator<AActor> It(World, AActor::StaticClass(), EActorIteratorFlags::SkipPendingKill); It; ++It)
 	{
 #if WITH_EDITOR
@@ -314,149 +322,6 @@
 	return SnapshotPropertyValue == OriginalPropertyValue;
 }
 
-void FWorldSnapshotData::ApplyToWorld_HandleRemovingActors(const FPropertySelectionMap& PropertiesToSerialize)
-{
-	for (const TWeakObjectPtr<AActor>& ActorToDespawn: PropertiesToSerialize.GetNewActorsToDespawn())
-	{
-		if (ensureMsgf(ActorToDespawn.IsValid(), TEXT("Actor became invalid since selection set was created")))
-		{
-			ActorToDespawn->Destroy(true, true);
-		}
-	}
-}
-
-void FWorldSnapshotData::ApplyToWorld_HandleRecreatingActors(TSet<AActor*>& EvaluatedActors, UPackage* LocalisationSnapshotPackage, const FPropertySelectionMap& PropertiesToSerialize)
-{
-#if WITH_EDITOR
-	FScopedSlowTask RecreateActors(PropertiesToSerialize.GetDeletedActorsToRespawn().Num(), LOCTEXT("ApplyToWorld.RecreateActorsKey", "Re-creating actors"));
-	RecreateActors.MakeDialogDelayed(1.f, false);
-#endif
-			
-	TMap<FSoftObjectPath, AActor*> RecreatedActors;
-	// 1st pass: allocate the actors. Serialisation is done in separate step so object references to other deleted actors resolve correctly.
-	for (const FSoftObjectPath& OriginalRemovedActorPath : PropertiesToSerialize.GetDeletedActorsToRespawn())
-	{
-		FActorSnapshotData* ActorSnapshot = ActorData.Find(OriginalRemovedActorPath);
-		if (!ensure(ActorSnapshot))
-		{
-			continue;	
-		}
-
-		UClass* ActorClass = ActorSnapshot->GetActorClass().ResolveClass();
-		if (!ensure(ActorClass))
-		{
-			UE_LOG(LogLevelSnapshots, Warning, TEXT("Failed to resolve class '%s'. Was it removed?"), *ActorSnapshot->GetActorClass().ToString());
-			continue;
-		}
-		
-		// The actor may have just been removed and still exist in memory
-		if (UObject* StillExistingRemovedActor = OriginalRemovedActorPath.ResolveObject())
-		{
-			AActor* Actor = Cast<AActor>(StillExistingRemovedActor);
-			// We need to rename the trash object otherwise SpawnActor will have a name collision and assert.
-			const FName NewName = MakeUniqueObjectName(Actor->GetWorld(), ActorClass, *Actor->GetName().Append(TEXT("_TRASH")));
-			Actor->Rename(*NewName.ToString());
-		}
-		
-		// Example: /Game/MapName.MapName:PersistentLevel.StaticMeshActor_42.StaticMeshComponent becomes /Game/MapName.MapName
-		const FSoftObjectPath PathToOwningWorldAsset = OriginalRemovedActorPath.GetAssetPathString();
-		UObject* UncastWorld = PathToOwningWorldAsset.ResolveObject();
-		if (!UncastWorld)
-		{
-			UncastWorld = PathToOwningWorldAsset.TryLoad();
-		}
-		if (!UncastWorld)
-		{
-			UE_LOG(LogLevelSnapshots, Error, TEXT("Failed to resolve world '%s'. The world should be loaded?"), *PathToOwningWorldAsset.ToString());
-			continue;
-		}
-
-		// Each Level in UWorld::Levels has a corresponding UWorld associated with it in which we re-create the actor.
-		if (UWorld* OwningLevelWorld = ExactCast<UWorld>(UncastWorld))
-		{
-			const FString& SubObjectPath = OriginalRemovedActorPath.GetSubPathString();
-			const int32 LastDotIndex = SubObjectPath.Find(TEXT("."));
-			// Full string: /Game/MapName.MapName:PersistentLevel.StaticMeshActor_42 SubPath: PersistentLevel.StaticMeshActor_42 
-			checkf(LastDotIndex != INDEX_NONE, TEXT("There should always be at least one dot after PersistentLevel"));
-			
-			const int32 NameLength = SubObjectPath.Len() - LastDotIndex - 1;
-			const FString ActorName = SubObjectPath.Right(NameLength);
-			
-			FActorSpawnParameters SpawnParameters;
-			SpawnParameters.Name = FName(*ActorName);
-			SpawnParameters.bNoFail = true;
-			SpawnParameters.Template = Cast<AActor>(GetClassDefault(ActorClass));
-			RecreatedActors.Add(OriginalRemovedActorPath, OwningLevelWorld->SpawnActor(ActorClass, nullptr, SpawnParameters));
-		}
-	}
-	
-	// 2nd pass: serialize. 
-	for (const FSoftObjectPath& OriginalRemovedActorPath : PropertiesToSerialize.GetDeletedActorsToRespawn())
-	{
-#if WITH_EDITOR
-		RecreateActors.EnterProgressFrame();
-#endif
-		if (FActorSnapshotData* ActorSnapshot = ActorData.Find(OriginalRemovedActorPath))
-		{
-			AActor** RecreatedActor = RecreatedActors.Find(OriginalRemovedActorPath);
-			if (ensure(RecreatedActor))
-			{
-				// Mark it, otherwise we'll serialize it again when we look for world actors matching the snapshot
-				EvaluatedActors.Add(*RecreatedActor);
-				ActorSnapshot->DeserializeIntoRecreatedEditorWorldActor(TempActorWorld.Get(), *RecreatedActor, *this, LocalisationSnapshotPackage, PropertiesToSerialize);
-			}
-		}
-	}
-}
-
-void FWorldSnapshotData::ApplyToWorld_HandleSerializingMatchingActors(TSet<AActor*>& EvaluatedActors, const TArray<FSoftObjectPath>& SelectedPaths, UPackage* LocalisationSnapshotPackage, const FPropertySelectionMap& PropertiesToSerialize)
-{
-#if WITH_EDITOR
-	FScopedSlowTask ExitingActorTask(SelectedPaths.Num(), LOCTEXT("ApplyToWorld.MatchingPropertiesKey", "Writing existing actors"));
-	ExitingActorTask.MakeDialogDelayed(1.f, true);
-#endif
-	for (const FSoftObjectPath& SelectedObject : SelectedPaths)
-	{
-#if WITH_EDITOR
-		ExitingActorTask.EnterProgressFrame();
-		if (ExitingActorTask.ShouldCancel())
-		{
-			return;
-		}
-#endif
-		
-		if (SelectedObject.IsValid())
-		{
-			AActor* OriginalWorldActor = nullptr;
-			UObject* ResolvedObject = SelectedObject.ResolveObject();
-			if (AActor* AsActor = Cast<AActor>(ResolvedObject))
-			{
-				OriginalWorldActor = AsActor;
-			}
-			else if (UActorComponent* AsComponent = Cast<UActorComponent>(ResolvedObject))
-			{
-				if (AActor* OwningActor = AsComponent->GetOwner())
-				{
-					OriginalWorldActor = OwningActor;
-				}
-			}
-
-			if (ensure(OriginalWorldActor))
-			{
-				if (!EvaluatedActors.Contains(OriginalWorldActor))
-				{
-					if (FActorSnapshotData* ActorSnapshot = ActorData.Find(OriginalWorldActor))
-					{
-						ActorSnapshot->DeserializeIntoExistingWorldActor(TempActorWorld.Get(), OriginalWorldActor, *this, LocalisationSnapshotPackage, PropertiesToSerialize);
-					}
-					
-					EvaluatedActors.Add(OriginalWorldActor);
-				}
-			}
-		}
-	}
-}
-
 int32 FWorldSnapshotData::AddObjectDependency(UObject* ReferenceFromOriginalObject)
 {
 	// TODO: Check whether subobject and track as such. Do this by walking up the chain of outers and determining whether it is owned by an actor or component.
@@ -477,10 +342,7 @@
 	const bool bIsExternalAssetReference = !SnapshotData.IsSet();
 	if (bIsExternalAssetReference)
 	{
-		// We're dealing with an external (asset) reference, e.g. a UMaterial.
-		UObject* ExternalReference = OriginalObjectPath.ResolveObject();
-		ensureAlwaysMsgf(ExternalReference == nullptr || (!ExternalReference->IsA<AActor>() && !ExternalReference->IsA<UActorComponent>()), TEXT("Something is wrong. We just checked that the reference is not a world object but it is an actor or component."));
-		return ExternalReference;
+		return ResolveExternalReference(OriginalObjectPath);
 	}
 
 	TOptional<AActor*> Result = SnapshotData.GetValue()->GetPreallocated(TempActorWorld->GetWorld(), *this);
@@ -519,16 +381,15 @@
 
 	bool bIsPathToActorSubobject;
 	const TOptional<FActorSnapshotData*> SnapshotData = FindSavedActorDataUsingObjectPath(ActorData, OriginalObjectPath, bIsPathToActorSubobject);
-	UObject* ResolvedObject = OriginalObjectPath.ResolveObject();
 	
 	// Dealing with an external (asset) reference, e.g. a UMaterial?
 	const bool bIsExternalAssetReference = !SnapshotData.IsSet();
 	if (bIsExternalAssetReference)
 	{
-		ensureAlwaysMsgf(ResolvedObject == nullptr || (!ResolvedObject->IsA<AActor>() && !ResolvedObject->IsA<UActorComponent>()), TEXT("Something is wrong. We just checked that the reference is not a world object but it is an actor or component."));
-		return ResolvedObject;
+		return ResolveExternalReference(OriginalObjectPath);
 	}
 
+	UObject* ResolvedObject = OriginalObjectPath.ResolveObject();
 	// Actor reference?
 	const bool bIsReferenceToActor = !bIsPathToActorSubobject;
 	if (bIsReferenceToActor)
@@ -537,7 +398,6 @@
 	}
 
 
-
 	const FSubobjectSnapshotData* SubobjectData = Subobjects.Find(ObjectPathIndex);
 	if (!SubobjectData)
 	{
@@ -564,13 +424,12 @@
 
 UObject* FWorldSnapshotData::ResolveObjectDependencyForClassDefaultObject(int32 ObjectPathIndex)
 {
-	if (!ensure(SerializedObjectReferences.IsValidIndex(ObjectPathIndex)))
+	if (ensure(SerializedObjectReferences.IsValidIndex(ObjectPathIndex)))
 	{
-		return nullptr;
+		const FSoftObjectPath& OriginalObjectPath = SerializedObjectReferences[ObjectPathIndex];
+		return ResolveExternalReference(OriginalObjectPath);
 	}
-	
-	const FSoftObjectPath& OriginalObjectPath = SerializedObjectReferences[ObjectPathIndex];
-	return OriginalObjectPath.ResolveObject();
+	return nullptr;
 }
 
 int32 FWorldSnapshotData::AddSubobjectDependency(UObject* ReferenceFromOriginalObject)
@@ -644,4 +503,174 @@
 	}
 }
 
+const FSnapshotVersionInfo& FWorldSnapshotData::GetSnapshotVersionInfo() const
+{
+	return SnapshotVersionInfo;
+}
+
+void FWorldSnapshotData::PostSerialize(const FArchive& Ar)
+{
+	if (Ar.IsLoading() && !SnapshotVersionInfo.IsInitialized())
+	{
+		SnapshotVersionInfo.Initialize();
+	}
+}
+
+UObject* FWorldSnapshotData::ResolveExternalReference(const FSoftObjectPath& ObjectPath)
+{
+	UObject* ExternalReference = ObjectPath.ResolveObject();
+	const bool bNeedsToLoadFromDisk = ExternalReference == nullptr;
+	if (bNeedsToLoadFromDisk)
+	{
+		ExternalReference = ObjectPath.TryLoad();
+	}
+		
+	// We're supposed to be dealing with an external (asset) reference, e.g. a UMaterial.
+	ensureAlwaysMsgf(ExternalReference == nullptr || (!ExternalReference->IsA<AActor>() && !ExternalReference->IsA<UActorComponent>()), TEXT("Something is wrong. We just checked that the reference is not a world object but it is an actor or component."));
+	return ExternalReference;
+}
+
+void FWorldSnapshotData::ApplyToWorld_HandleRemovingActors(const FPropertySelectionMap& PropertiesToSerialize)
+{
+	for (const TWeakObjectPtr<AActor>& ActorToDespawn: PropertiesToSerialize.GetNewActorsToDespawn())
+	{
+		if (ensureMsgf(ActorToDespawn.IsValid(), TEXT("Actor became invalid since selection set was created")))
+		{
+			ActorToDespawn->Destroy(true, true);
+		}
+	}
+}
+
+void FWorldSnapshotData::ApplyToWorld_HandleRecreatingActors(TSet<AActor*>& EvaluatedActors, UPackage* LocalisationSnapshotPackage, const FPropertySelectionMap& PropertiesToSerialize)
+{
+#if WITH_EDITOR
+	FScopedSlowTask RecreateActors(PropertiesToSerialize.GetDeletedActorsToRespawn().Num(), LOCTEXT("ApplyToWorld.RecreateActorsKey", "Re-creating actors"));
+	RecreateActors.MakeDialogDelayed(1.f, false);
+#endif
+			
+	TMap<FSoftObjectPath, AActor*> RecreatedActors;
+	// 1st pass: allocate the actors. Serialisation is done in separate step so object references to other deleted actors resolve correctly.
+	for (const FSoftObjectPath& OriginalRemovedActorPath : PropertiesToSerialize.GetDeletedActorsToRespawn())
+	{
+		FActorSnapshotData* ActorSnapshot = ActorData.Find(OriginalRemovedActorPath);
+		if (!ensure(ActorSnapshot))
+		{
+			continue;	
+		}
+
+		UClass* ActorClass = ActorSnapshot->GetActorClass().ResolveClass();
+		if (!ensure(ActorClass))
+		{
+			UE_LOG(LogLevelSnapshots, Warning, TEXT("Failed to resolve class '%s'. Was it removed?"), *ActorSnapshot->GetActorClass().ToString());
+			continue;
+		}
+		
+		// The actor may have just been removed and still exist in memory
+		if (UObject* StillExistingRemovedActor = OriginalRemovedActorPath.ResolveObject())
+		{
+			AActor* Actor = Cast<AActor>(StillExistingRemovedActor);
+			// We need to rename the trash object otherwise SpawnActor will have a name collision and assert.
+			const FName NewName = MakeUniqueObjectName(Actor->GetWorld(), ActorClass, *Actor->GetName().Append(TEXT("_TRASH")));
+			Actor->Rename(*NewName.ToString());
+		}
+		
+		// Example: /Game/MapName.MapName:PersistentLevel.StaticMeshActor_42.StaticMeshComponent becomes /Game/MapName.MapName
+		const FSoftObjectPath PathToOwningWorldAsset = OriginalRemovedActorPath.GetAssetPathString();
+		UObject* UncastWorld = PathToOwningWorldAsset.ResolveObject();
+		if (!UncastWorld)
+		{
+			UncastWorld = PathToOwningWorldAsset.TryLoad();
+		}
+		if (!UncastWorld)
+		{
+			UE_LOG(LogLevelSnapshots, Error, TEXT("Failed to resolve world '%s'. The world should be loaded?"), *PathToOwningWorldAsset.ToString());
+			continue;
+		}
+
+		// Each Level in UWorld::Levels has a corresponding UWorld associated with it in which we re-create the actor.
+		if (UWorld* OwningLevelWorld = ExactCast<UWorld>(UncastWorld))
+		{
+			const FString& SubObjectPath = OriginalRemovedActorPath.GetSubPathString();
+			const int32 LastDotIndex = SubObjectPath.Find(TEXT("."));
+			// Full string: /Game/MapName.MapName:PersistentLevel.StaticMeshActor_42 SubPath: PersistentLevel.StaticMeshActor_42 
+			checkf(LastDotIndex != INDEX_NONE, TEXT("There should always be at least one dot after PersistentLevel"));
+			
+			const int32 NameLength = SubObjectPath.Len() - LastDotIndex - 1;
+			const FString ActorName = SubObjectPath.Right(NameLength);
+			
+			FActorSpawnParameters SpawnParameters;
+			SpawnParameters.Name = FName(*ActorName);
+			SpawnParameters.bNoFail = true;
+			SpawnParameters.Template = Cast<AActor>(GetClassDefault(ActorClass));
+			RecreatedActors.Add(OriginalRemovedActorPath, OwningLevelWorld->SpawnActor(ActorClass, nullptr, SpawnParameters));
+		}
+	}
+	
+	// 2nd pass: serialize. 
+	for (const FSoftObjectPath& OriginalRemovedActorPath : PropertiesToSerialize.GetDeletedActorsToRespawn())
+	{
+#if WITH_EDITOR
+		RecreateActors.EnterProgressFrame();
+#endif
+		if (FActorSnapshotData* ActorSnapshot = ActorData.Find(OriginalRemovedActorPath))
+		{
+			AActor** RecreatedActor = RecreatedActors.Find(OriginalRemovedActorPath);
+			if (ensure(RecreatedActor))
+			{
+				// Mark it, otherwise we'll serialize it again when we look for world actors matching the snapshot
+				EvaluatedActors.Add(*RecreatedActor);
+				ActorSnapshot->DeserializeIntoRecreatedEditorWorldActor(TempActorWorld.Get(), *RecreatedActor, *this, LocalisationSnapshotPackage, PropertiesToSerialize);
+			}
+		}
+	}
+}
+
+void FWorldSnapshotData::ApplyToWorld_HandleSerializingMatchingActors(TSet<AActor*>& EvaluatedActors, const TArray<FSoftObjectPath>& SelectedPaths, UPackage* LocalisationSnapshotPackage, const FPropertySelectionMap& PropertiesToSerialize)
+{
+#if WITH_EDITOR
+	FScopedSlowTask ExitingActorTask(SelectedPaths.Num(), LOCTEXT("ApplyToWorld.MatchingPropertiesKey", "Writing existing actors"));
+	ExitingActorTask.MakeDialogDelayed(1.f, true);
+#endif
+	for (const FSoftObjectPath& SelectedObject : SelectedPaths)
+	{
+#if WITH_EDITOR
+		ExitingActorTask.EnterProgressFrame();
+		if (ExitingActorTask.ShouldCancel())
+		{
+			return;
+		}
+#endif
+		
+		if (SelectedObject.IsValid())
+		{
+			AActor* OriginalWorldActor = nullptr;
+			UObject* ResolvedObject = SelectedObject.ResolveObject();
+			if (AActor* AsActor = Cast<AActor>(ResolvedObject))
+			{
+				OriginalWorldActor = AsActor;
+			}
+			else if (UActorComponent* AsComponent = Cast<UActorComponent>(ResolvedObject))
+			{
+				if (AActor* OwningActor = AsComponent->GetOwner())
+				{
+					OriginalWorldActor = OwningActor;
+				}
+			}
+
+			if (ensure(OriginalWorldActor))
+			{
+				if (!EvaluatedActors.Contains(OriginalWorldActor))
+				{
+					if (FActorSnapshotData* ActorSnapshot = ActorData.Find(OriginalWorldActor))
+					{
+						ActorSnapshot->DeserializeIntoExistingWorldActor(TempActorWorld.Get(), OriginalWorldActor, *this, LocalisationSnapshotPackage, PropertiesToSerialize);
+					}
+					
+					EvaluatedActors.Add(OriginalWorldActor);
+				}
+			}
+		}
+	}
+}
+
 #undef LOCTEXT_NAMESPACE
diff --git a/Engine/Plugins/VirtualProduction/LevelSnapshots/Source/LevelSnapshots/Private/LevelSnapshotsFunctionLibrary.cpp b/Engine/Plugins/VirtualProduction/LevelSnapshots/Source/LevelSnapshots/Private/LevelSnapshotsFunctionLibrary.cpp
--- a/Engine/Plugins/VirtualProduction/LevelSnapshots/Source/LevelSnapshots/Private/LevelSnapshotsFunctionLibrary.cpp
+++ b/Engine/Plugins/VirtualProduction/LevelSnapshots/Source/LevelSnapshots/Private/LevelSnapshotsFunctionLibrary.cpp
@@ -115,7 +115,7 @@ ULevelSnapshot* ULevelSnapshotsFunctionLibrary::TakeLevelSnapshot_Internal(const
 		return nullptr;
 	}
 	
-	ULevelSnapshot* NewSnapshot = NewObject<ULevelSnapshot>(InPackage ? InPackage : GetTransientPackage(), NewSnapshotName, RF_Public | RF_Standalone);
+	ULevelSnapshot* NewSnapshot = NewObject<ULevelSnapshot>(InPackage ? InPackage : GetTransientPackage(), NewSnapshotName, RF_NoFlags);
 	NewSnapshot->SetSnapshotName(NewSnapshotName);
 	NewSnapshot->SetSnapshotDescription(Description);
 	NewSnapshot->SnapshotWorld(TargetWorld);
diff --git a/Engine/Plugins/VirtualProduction/LevelSnapshots/Source/LevelSnapshots/Public/Data/SnapshotVersion.h b/Engine/Plugins/VirtualProduction/LevelSnapshots/Source/LevelSnapshots/Public/Data/SnapshotVersion.h
new file mode 100644
--- /dev/null
+++ b/Engine/Plugins/VirtualProduction/LevelSnapshots/Source/LevelSnapshots/Public/Data/SnapshotVersion.h
@@ -0,0 +1,100 @@
+// Copyright Epic Games, Inc. All Rights Reserved.
+
+#pragma once
+
+#include "CoreMinimal.h"
+#include "SnapshotVersion.generated.h"
+
+class FEngineVersion;
+struct FCustomVersion;
+
+
+/** Holds file version information */
+USTRUCT()
+struct LEVELSNAPSHOTS_API FSnapshotFileVersionInfo
+{
+	GENERATED_BODY()
+
+	/** Initialize this version info from the compiled in data */
+	void Initialize();
+
+	/* UE4 file version */
+	UPROPERTY()
+	int32 FileVersionUE4 = 0;
+
+	/* Licensee file version */
+	UPROPERTY()
+	int32 FileVersionLicenseeUE4 = 0;
+};
+
+/** Holds engine version information */
+USTRUCT()
+struct LEVELSNAPSHOTS_API FSnapshotEngineVersionInfo
+{
+	GENERATED_BODY()
+
+	/** Initialize this version info from the given version */
+	void Initialize(const FEngineVersion& InVersion);
+
+	/** Major version number */
+	UPROPERTY()
+	uint16 Major = 0;
+
+	/** Minor version number */
+	UPROPERTY()
+	uint16 Minor = 0;
+
+	/** Patch version number */
+	UPROPERTY()
+	uint16 Patch = 0;
+
+	/** Changelist number. This is used to arbitrate when Major/Minor/Patch version numbers match */
+	UPROPERTY()
+	uint32 Changelist = 0;
+};
+
+/** Holds custom version information */
+USTRUCT()
+struct LEVELSNAPSHOTS_API FSnapshotCustomVersionInfo
+{
+	GENERATED_BODY()
+
+	/** Initialize this version info from the given version */
+	void Initialize(const FCustomVersion& InVersion);
+
+	/** Friendly name of the version */
+	UPROPERTY()
+	FName FriendlyName;
+
+	/** Unique custom key */
+	UPROPERTY()
+	FGuid Key;
+
+	/** Custom version */
+	UPROPERTY()
+	int32 Version = 0;
+};
+
+/** Holds version information for a session */
+USTRUCT()
+struct LEVELSNAPSHOTS_API FSnapshotVersionInfo
+{
+	GENERATED_BODY()
+
+	/** Initialize this version info from the compiled in data */
+	void Initialize();
+
+	bool IsInitialized() const;
+
+	/** File version info */
+	UPROPERTY()
+	FSnapshotFileVersionInfo FileVersion;
+
+	/** Engine version info */
+	UPROPERTY()
+	FSnapshotEngineVersionInfo EngineVersion;
+
+	/** Custom version info */
+	UPROPERTY()
+	TArray<FSnapshotCustomVersionInfo> CustomVersions;
+};
diff --git a/Engine/Plugins/VirtualProduction/LevelSnapshots/Source/LevelSnapshots/Public/Data/WorldSnapshotData.h b/Engine/Plugins/VirtualProduction/LevelSnapshots/Source/LevelSnapshots/Public/Data/WorldSnapshotData.h
--- a/Engine/Plugins/VirtualProduction/LevelSnapshots/Source/LevelSnapshots/Public/Data/WorldSnapshotData.h
+++ b/Engine/Plugins/VirtualProduction/LevelSnapshots/Source/LevelSnapshots/Public/Data/WorldSnapshotData.h
@@ -6,7 +6,9 @@
 #include "ActorSnapshotData.h"
 #include "ComponentSnapshotData.h"
 #include "ClassDefaultObjectSnapshotData.h"
+#include "SnapshotVersion.h"
 #include "SubobjectSnapshotData.h"
+#include "UObject/Class.h"
 #include "WorldSnapshotData.generated.h"
 
 class AActor;
@@ -43,9 +45,8 @@ struct LEVELSNAPSHOTS_API FWorldSnapshotData
 	 * Checks whether two pointers point to "equivalent" objects.
 	 */
 	bool AreReferencesEquivalent(UObject* SnapshotPropertyValue, UObject* OriginalPropertyValue, AActor* SnapshotActor, AActor* OriginalActor) const;
-
+	
 public: /* Serialisation functions */
-
 	
 	
 	/* Adds an object dependency without serializing the object's content. Intended for external objects, e.g. to UMaterial in the content browser. */
@@ -79,9 +80,18 @@ public: /* Serialisation functions */
 	* This is intended for cases where you cannot specify a template object for new objects. Components are one such use case.
 	*/
 	void SerializeClassDefaultsInto(UObject* Object);
+
+
+	const FSnapshotVersionInfo& GetSnapshotVersionInfo() const;
+	
+	//~ Begin TStructOpsTypeTraits Interface
+	void PostSerialize(const FArchive& Ar);
+	//~ End TStructOpsTypeTraits Interface
 	
 private:
 
+	UObject* ResolveExternalReference(const FSoftObjectPath& ObjectPath);
+	
 	void ApplyToWorld_HandleRemovingActors(const FPropertySelectionMap& PropertiesToSerialize);
 	void ApplyToWorld_HandleRecreatingActors(TSet<AActor*>& EvaluatedActors, UPackage* LocalisationSnapshotPackage, const FPropertySelectionMap& PropertiesToSerialize);
 	void ApplyToWorld_HandleSerializingMatchingActors(TSet<AActor*>& EvaluatedActors, const TArray<FSoftObjectPath>& SelectedPaths, UPackage* LocalisationSnapshotPackage, const FPropertySelectionMap& PropertiesToSerialize);
@@ -91,7 +101,14 @@ private:
 	/* The world we will be adding temporary actors to */
 	TWeakObjectPtr<UWorld> TempActorWorld;
 
-
+	/**
+	 * Stores versioning information we inject into archives.
+	 * This is to support asset migration, like FArchive::UsingCustomVersion.
+	 */
+	UPROPERTY()
+	FSnapshotVersionInfo SnapshotVersionInfo;
+	
+	
 	
 	/* We only save properties with values different from their CDO counterpart.
 	 * Because of this, we need to save class defaults in the snapshot.
@@ -130,5 +147,14 @@ private:
 	 */
 	UPROPERTY()
 	TMap<int32, FSubobjectSnapshotData> Subobjects;
-	
 };
+
+template<>
+struct TStructOpsTypeTraits<FWorldSnapshotData> : public TStructOpsTypeTraitsBase2<FWorldSnapshotData>
+{
+	enum 
+	{ 
+		WithPostSerialize = true
+	};
+};
+
diff --git a/Engine/Plugins/VirtualProduction/LevelSnapshots/Source/LevelSnapshotsEditor/Private/Customizations/NegatableFilterDetailsCustomization.cpp b/Engine/Plugins/VirtualProduction/LevelSnapshots/Source/LevelSnapshotsEditor/Private/Customizations/NegatableFilterDetailsCustomization.cpp
--- a/Engine/Plugins/VirtualProduction/LevelSnapshots/Source/LevelSnapshotsEditor/Private/Customizations/NegatableFilterDetailsCustomization.cpp
+++ b/Engine/Plugins/VirtualProduction/LevelSnapshots/Source/LevelSnapshotsEditor/Private/Customizations/NegatableFilterDetailsCustomization.cpp
@@ -6,6 +6,7 @@
 
 #include "DetailLayoutBuilder.h"
 #include "DetailCategoryBuilder.h"
+#include "Algo/AllOf.h"
 #include "Algo/Transform.h"
 
 namespace
@@ -22,9 +23,11 @@ namespace
 		DetailBuilder.GetObjectsBeingCustomized(WeakSelectedObjects);
 
 		TArray<UObject*> SelectedObjects;
-		Algo::Transform(WeakSelectedObjects, SelectedObjects, [&SelectedObjects](TWeakObjectPtr<UObject> Obj) { return CastChecked<UNegatableFilter>(Obj.Get())->GetChildFilter(); });
-
-		if (!ensure(SelectedObjects.Num() != 0))
+		Algo::Transform(WeakSelectedObjects, SelectedObjects, [&SelectedObjects](TWeakObjectPtr<UObject> Obj) { return ExactCast<UNegatableFilter>(Obj.Get())->GetChildFilter(); });
+		
+		// Child filter may be null in certain cases, e.g. when user force deleted the filter class
+		const bool bNoneAreNull = Algo::AllOf(SelectedObjects, [](UObject* ChildFilter) { return ChildFilter != nullptr; });
+		if (!bNoneAreNull || !ensure(SelectedObjects.Num() != 0))
 		{
 			return;
 		}
diff --git a/Engine/Plugins/VirtualProduction/LevelSnapshots/Source/LevelSnapshotsEditor/Private/Data/FavoriteFilterContainer.cpp b/Engine/Plugins/VirtualProduction/LevelSnapshots/Source/LevelSnapshotsEditor/Private/Data/FavoriteFilterContainer.cpp
--- a/Engine/Plugins/VirtualProduction/LevelSnapshots/Source/LevelSnapshotsEditor/Private/Data/FavoriteFilterContainer.cpp
+++ b/Engine/Plugins/VirtualProduction/LevelSnapshots/Source/LevelSnapshotsEditor/Private/Data/FavoriteFilterContainer.cpp
@@ -62,6 +62,21 @@ namespace
 		return Result;
 	}
 
+	bool IsBlueprintFilter(const FAssetData& BlueprintClassData)
+	{
+		UClass* BlueprintFilterClass = ULevelSnapshotBlueprintFilter::StaticClass();
+		
+		const FString NativeParentClassPath = BlueprintClassData.GetTagValueRef<FString>(FBlueprintTags::NativeParentClassPath);
+		const FSoftClassPath ClassPath(NativeParentClassPath);
+			
+		UClass* NativeParentClass = ClassPath.ResolveClass();
+		const bool bInheritsFromBlueprintFilter =
+			NativeParentClass // Class may have been removed, or renamed and not correctly redirected
+			&& (NativeParentClass == BlueprintFilterClass || NativeParentClass->IsChildOf(BlueprintFilterClass));
+
+		return bInheritsFromBlueprintFilter;
+	}
+	
 	TArray<TSubclassOf<ULevelSnapshotBlueprintFilter>> FindBlueprintFilterClasses()
 	{
 		FAssetRegistryModule& AssetRegistryModule = FModuleManager::LoadModuleChecked<FAssetRegistryModule>(TEXT("AssetRegistry"));
@@ -71,16 +86,10 @@ namespace
 		AssetRegistry.GetAssetsByClass(UBlueprint::StaticClass()->GetFName(), BlueprintList);
 
 		// Returns = 
-		const FString BlueprintFilterBaseClassPath = ULevelSnapshotBlueprintFilter::StaticClass()->GetPathName();
 		TArray<TSubclassOf<ULevelSnapshotBlueprintFilter>> Result;
 		for (const FAssetData& BlueprintClassData : BlueprintList)
 		{
-			// Is of form = Class''
-			const FString FirstNativeParentName = BlueprintClassData.GetTagValueRef<FString>(FBlueprintTags::NativeParentClassPath);
-			const bool bInheritsFromBlueprintFilter = FirstNativeParentName.Contains(BlueprintFilterBaseClassPath);
-
-			// Loading all blueprints is too slow... this finds only the ones we need
-			if (!bInheritsFromBlueprintFilter)
+			if (!IsBlueprintFilter(BlueprintClassData))
 			{
 				continue;
 			}
@@ -88,7 +97,7 @@ namespace
 			// If there are a lot of assets, this may be slow.
 			UBlueprint* BlueprintAsset = Cast<UBlueprint>(BlueprintClassData.GetAsset());
 			UClass* LoadedClass = BlueprintAsset->GeneratedClass;
-			if (ensure(LoadedClass) && ensure(BlueprintAsset->ParentClass) && BlueprintAsset->ParentClass->IsChildOf(ULevelSnapshotBlueprintFilter::StaticClass()))
+			if (ensure(LoadedClass && BlueprintAsset->ParentClass))
 			{
 				Result.Add(LoadedClass);
 			}
@@ -97,6 +106,31 @@ namespace
 	}
 }
 
+void UFavoriteFilterContainer::PostInitProperties()
+{
+	Super::PostInitProperties();
+
+	IAssetRegistry* AssetRegistry = IAssetRegistry::Get();
+	if (ensure(AssetRegistry))
+	{
+		OnAssetDeleted = AssetRegistry->OnAssetRemoved().AddLambda([this](const FAssetData& AssetData)
+		{
+			CleanseFavorites();
+		});
+	}
+}
+
+void UFavoriteFilterContainer::BeginDestroy()
+{
+	IAssetRegistry* AssetRegistry = IAssetRegistry::Get();
+	if (AssetRegistry)
+	{
+		AssetRegistry->OnAssetRemoved().Remove(OnAssetDeleted);
+	}
+	
+	Super::BeginDestroy();
+}
+
 void UFavoriteFilterContainer::AddToFavorites(const TSubclassOf<ULevelSnapshotFilter>& NewFavoriteClass)
 {
 	if (!ensure(NewFavoriteClass))
@@ -110,6 +144,8 @@ void UFavoriteFilterContainer::AddToFavorites(const TSubclassOf<ULevelSnapshotFi
 	{
 		OnFavoritesChanged.Broadcast();
 	}
+	
+	CleanseFavorites();
 }
 
 void UFavoriteFilterContainer::RemoveFromFavorites(const TSubclassOf<ULevelSnapshotFilter>& NoLongerFavoriteClass)
@@ -135,6 +171,8 @@ void UFavoriteFilterContainer::RemoveFromFavorites(const TSubclassOf<ULevelSnaps
 	{
 		OnFavoritesChanged.Broadcast();
 	}
+
+	CleanseFavorites();
 }
 
 void UFavoriteFilterContainer::ClearFavorites()
@@ -224,6 +262,24 @@ TArray<TSubclassOf<ULevelSnapshotFilter>> UFavoriteFilterContainer::GetFiltersIn
 	return {};
 }
 
+void UFavoriteFilterContainer::CleanseFavorites()
+{
+	const int32 NumBeforeRemoval = Favorites.Num();
+	for (int32 i = 0; i < Favorites.Num(); ++i)
+	{
+		if (Favorites[i] == nullptr)
+		{
+			Favorites.RemoveAt(i);
+			--i;
+		}
+	}
+
+	if (NumBeforeRemoval != Favorites.Num())
+	{
+		OnFavoritesChanged.Broadcast();
+	}
+}
+
 void UFavoriteFilterContainer::SetIncludeAllNativeClasses(bool bShouldIncludeNative)
 {
 	if (bShouldIncludeNative == bIncludeAllNativeClasses)
diff --git a/Engine/Plugins/VirtualProduction/LevelSnapshots/Source/LevelSnapshotsEditor/Private/Data/FavoriteFilterContainer.h b/Engine/Plugins/VirtualProduction/LevelSnapshots/Source/LevelSnapshotsEditor/Private/Data/FavoriteFilterContainer.h
--- a/Engine/Plugins/VirtualProduction/LevelSnapshots/Source/LevelSnapshotsEditor/Private/Data/FavoriteFilterContainer.h
+++ b/Engine/Plugins/VirtualProduction/LevelSnapshots/Source/LevelSnapshotsEditor/Private/Data/FavoriteFilterContainer.h
@@ -17,6 +17,11 @@ class UFavoriteFilterContainer : public UObject
 {
 	GENERATED_BODY()
 public:
+
+	//~ Begin UObject Interface
+	virtual void PostInitProperties() override;
+	virtual void BeginDestroy() override;
+	//~ Begin UObject Interface
 	
 	void AddToFavorites(const TSubclassOf<ULevelSnapshotFilter>& NewFavoriteClass);
 	void RemoveFromFavorites(const TSubclassOf<ULevelSnapshotFilter>& NoLongerFavoriteClass);
@@ -40,6 +45,8 @@ public:
 	
 private:
 
+	void CleanseFavorites();
+	
 	void SetIncludeAllNativeClasses(bool bShouldIncludeNative);
 	void SetIncludeAllBlueprintClasses(bool bShouldIncludeBlueprint);
 	bool ShouldIncludeAllNativeClasses() const;
@@ -52,11 +59,13 @@ private:
 
 	
 	
-	// TODO: detect when blueprint class is deleted and remove it from Favorites and LoadedBlueprintClasses
-	
 	/* The filters the user selected to use. */
+	UPROPERTY()
 	TArray<TSubclassOf<ULevelSnapshotFilter>> Favorites;
+	
 	bool bIncludeAllNativeClasses = false;
 	bool bIncludeAllBlueprintClasses = false;
-	
+
+	/** So we can remove any filter classes that were favourited */
+	FDelegateHandle OnAssetDeleted;
 };
\ No newline at end of file
diff --git a/Engine/Plugins/VirtualProduction/LevelSnapshots/Source/LevelSnapshotsEditor/Private/Data/Filters/ConjunctionFilter.cpp b/Engine/Plugins/VirtualProduction/LevelSnapshots/Source/LevelSnapshotsEditor/Private/Data/Filters/ConjunctionFilter.cpp
--- a/Engine/Plugins/VirtualProduction/LevelSnapshots/Source/LevelSnapshotsEditor/Private/Data/Filters/ConjunctionFilter.cpp
+++ b/Engine/Plugins/VirtualProduction/LevelSnapshots/Source/LevelSnapshotsEditor/Private/Data/Filters/ConjunctionFilter.cpp
@@ -95,6 +95,22 @@ EFilterResult::Type UConjunctionFilter::IsPropertyValid(const FIsPropertyValidPa
 		});
 }
 
+EFilterResult::Type UConjunctionFilter::IsDeletedActorValid(const FIsDeletedActorValidParams& Params) const
+{
+	return ExecuteAndChain(Children, EditorFilterBehavior, [&Params](UNegatableFilter* Child)
+		{
+			return Child->IsDeletedActorValid(Params);
+		});
+}
+
+EFilterResult::Type UConjunctionFilter::IsAddedActorValid(const FIsAddedActorValidParams& Params) const
+{
+	return ExecuteAndChain(Children, EditorFilterBehavior, [&Params](UNegatableFilter* Child)
+		{
+			return Child->IsAddedActorValid(Params);
+		});
+}
+
 TArray<UEditorFilter*> UConjunctionFilter::GetEditorChildren()
 {
 	TArray<UEditorFilter*> EditorFilterChildren;
diff --git a/Engine/Plugins/VirtualProduction/LevelSnapshots/Source/LevelSnapshotsEditor/Private/Data/Filters/ConjunctionFilter.h b/Engine/Plugins/VirtualProduction/LevelSnapshots/Source/LevelSnapshotsEditor/Private/Data/Filters/ConjunctionFilter.h
--- a/Engine/Plugins/VirtualProduction/LevelSnapshots/Source/LevelSnapshotsEditor/Private/Data/Filters/ConjunctionFilter.h
+++ b/Engine/Plugins/VirtualProduction/LevelSnapshots/Source/LevelSnapshotsEditor/Private/Data/Filters/ConjunctionFilter.h
@@ -28,10 +28,12 @@ public:
 	const TArray<UNegatableFilter*>& GetChildren() const;
 	
 	//~ Begin ULevelSnapshotFilter Interface
-	EFilterResult::Type IsActorValid(const FIsActorValidParams& Params) const override;
-	EFilterResult::Type IsPropertyValid(const FIsPropertyValidParams& Params) const override;
+	virtual EFilterResult::Type IsActorValid(const FIsActorValidParams& Params) const override;
+	virtual EFilterResult::Type IsPropertyValid(const FIsPropertyValidParams& Params) const override;
+	virtual EFilterResult::Type IsDeletedActorValid(const FIsDeletedActorValidParams& Params) const override;
+	virtual EFilterResult::Type IsAddedActorValid(const FIsAddedActorValidParams& Params) const override;
 	//~ End ULevelSnapshotFilter Interface
-
+	
 	//~ Begin UEditorFilter Interface
 	virtual TArray<UEditorFilter*> GetEditorChildren();
 	virtual void IncrementEditorFilterBehavior(const bool bIncludeChildren = false) override;
diff --git a/Engine/Plugins/VirtualProduction/LevelSnapshots/Source/LevelSnapshotsEditor/Private/Data/Filters/DisjunctiveNormalFormFilter.cpp b/Engine/Plugins/VirtualProduction/LevelSnapshots/Source/LevelSnapshotsEditor/Private/Data/Filters/DisjunctiveNormalFormFilter.cpp
--- a/Engine/Plugins/VirtualProduction/LevelSnapshots/Source/LevelSnapshotsEditor/Private/Data/Filters/DisjunctiveNormalFormFilter.cpp
+++ b/Engine/Plugins/VirtualProduction/LevelSnapshots/Source/LevelSnapshotsEditor/Private/Data/Filters/DisjunctiveNormalFormFilter.cpp
@@ -98,6 +98,22 @@ EFilterResult::Type UDisjunctiveNormalFormFilter::IsPropertyValid(const FIsPrope
 	});
 }
 
+EFilterResult::Type UDisjunctiveNormalFormFilter::IsDeletedActorValid(const FIsDeletedActorValidParams& Params) const
+{
+	return ExecuteOrChain(Children, [&Params](UConjunctionFilter* Child)
+	{
+		return Child->IsDeletedActorValid(Params);
+	});
+}
+
+EFilterResult::Type UDisjunctiveNormalFormFilter::IsAddedActorValid(const FIsAddedActorValidParams& Params) const
+{
+	return ExecuteOrChain(Children, [&Params](UConjunctionFilter* Child)
+	{
+		return Child->IsAddedActorValid(Params);
+	});
+}
+
 TArray<UEditorFilter*> UDisjunctiveNormalFormFilter::GetEditorChildren()
 {
 	TArray<UEditorFilter*> EditorFilterChildren;
diff --git a/Engine/Plugins/VirtualProduction/LevelSnapshots/Source/LevelSnapshotsEditor/Private/Data/Filters/DisjunctiveNormalFormFilter.h b/Engine/Plugins/VirtualProduction/LevelSnapshots/Source/LevelSnapshotsEditor/Private/Data/Filters/DisjunctiveNormalFormFilter.h
--- a/Engine/Plugins/VirtualProduction/LevelSnapshots/Source/LevelSnapshotsEditor/Private/Data/Filters/DisjunctiveNormalFormFilter.h
+++ b/Engine/Plugins/VirtualProduction/LevelSnapshots/Source/LevelSnapshotsEditor/Private/Data/Filters/DisjunctiveNormalFormFilter.h
@@ -28,10 +28,12 @@ public:
 	const TArray<UConjunctionFilter*>& GetChildren() const;
 
 	//~ Begin ULevelSnapshotFilter Interface
-	EFilterResult::Type IsActorValid(const FIsActorValidParams& Params) const override;
-	EFilterResult::Type IsPropertyValid(const FIsPropertyValidParams& Params) const override;
+	virtual EFilterResult::Type IsActorValid(const FIsActorValidParams& Params) const override;
+	virtual EFilterResult::Type IsPropertyValid(const FIsPropertyValidParams& Params) const override;
+	virtual EFilterResult::Type IsDeletedActorValid(const FIsDeletedActorValidParams& Params) const override;
+	virtual EFilterResult::Type IsAddedActorValid(const FIsAddedActorValidParams& Params) const override;
 	//~ End ULevelSnapshotFilter Interface
-
+	
 	//~ Begin UEditorFilter Interface
 	virtual TArray<UEditorFilter*> GetEditorChildren();
 	//~ Begin UEditorFilter Interface
diff --git a/Engine/Plugins/VirtualProduction/LevelSnapshots/Source/LevelSnapshotsEditor/Private/Data/Filters/NegatableFilter.cpp b/Engine/Plugins/VirtualProduction/LevelSnapshots/Source/LevelSnapshotsEditor/Private/Data/Filters/NegatableFilter.cpp
--- a/Engine/Plugins/VirtualProduction/LevelSnapshots/Source/LevelSnapshotsEditor/Private/Data/Filters/NegatableFilter.cpp
+++ b/Engine/Plugins/VirtualProduction/LevelSnapshots/Source/LevelSnapshotsEditor/Private/Data/Filters/NegatableFilter.cpp
@@ -2,6 +2,21 @@
 
 #include "NegatableFilter.h"
 
+namespace
+{
+	EFilterResult::Type ApplyFilter(EEditorFilterBehavior NegationBehaviour, TFunction<EFilterResult::Type()> ApplyFilter)
+	{
+		if (NegationBehaviour == EEditorFilterBehavior::Ignore)
+		{
+			return EFilterResult::DoNotCare; 
+		}
+
+		const EFilterResult::Type Result = ApplyFilter();
+		const bool bShouldNegate = NegationBehaviour == EEditorFilterBehavior::Negate;
+		return bShouldNegate ? EFilterResult::Negate(Result) : Result;
+	}
+}
+
 UNegatableFilter* UNegatableFilter::CreateNegatableFilter(ULevelSnapshotFilter* ChildFilter, const TOptional<UObject*>& Outer)
 {
 	if (!ensure(ChildFilter))
@@ -22,38 +37,48 @@ ULevelSnapshotFilter* UNegatableFilter::GetChildFilter() const
 
 FText UNegatableFilter::GetDisplayName() const
 {
+	// Child filter is nullptr when user deletes filter class and force deletes
+	if (!ChildFilter)
+	{
+		return FText::FromString(TEXT("Missing Filter Class"));
+	}
+
 	const bool bHasName = Name.Len() != 0;
 	return bHasName ? FText::FromString(Name) : ChildFilter->GetClass()->GetDisplayNameText();
 }
 
 EFilterResult::Type UNegatableFilter::IsActorValid(const FIsActorValidParams& Params) const
 {
-	if (EditorFilterBehavior == EEditorFilterBehavior::Ignore)
+	return ApplyFilter(EditorFilterBehavior, [this, &Params]()
 	{
-		return EFilterResult::DoNotCare; // This filter is always used in an AND-chain
-	}
+		// Child filter is nullptr when user deletes filter class and force deletes
+		return ChildFilter ? ChildFilter->IsActorValid(Params) : EFilterResult::DoNotCare;
+	});
+}
 
-	if (ensure(ChildFilter))
+EFilterResult::Type UNegatableFilter::IsPropertyValid(const FIsPropertyValidParams& Params) const
+{
+	return ApplyFilter(EditorFilterBehavior, [this, &Params]()
 	{
-		EFilterResult::Type Result = ChildFilter->IsActorValid(Params);
-		const bool bShouldNegate = EditorFilterBehavior == EEditorFilterBehavior::Negate;
-		return bShouldNegate ? EFilterResult::Negate(Result) : Result;
-	}
-	return EFilterResult::DoNotCare;
+		// Child filter is nullptr when user deletes filter class and force deletes
+		return ChildFilter ? ChildFilter->IsPropertyValid(Params) : EFilterResult::DoNotCare;
+	});
 }
 
-EFilterResult::Type UNegatableFilter::IsPropertyValid(const FIsPropertyValidParams& Params) const
+EFilterResult::Type UNegatableFilter::IsDeletedActorValid(const FIsDeletedActorValidParams& Params) const
 {
-	if (EditorFilterBehavior == EEditorFilterBehavior::Ignore)
+	return ApplyFilter(EditorFilterBehavior, [this, &Params]()
 	{
-		return EFilterResult::DoNotCare; // This filter is always used in an AND-chain
-	}
-	
-	if (ensure(ChildFilter))
+		// Child filter is nullptr when user deletes filter class and force deletes
+		return ChildFilter ? ChildFilter->IsDeletedActorValid(Params) : EFilterResult::DoNotCare;
+	});
+}
+
+EFilterResult::Type UNegatableFilter::IsAddedActorValid(const FIsAddedActorValidParams& Params) const
+{
+	return ApplyFilter(EditorFilterBehavior, [this, &Params]()
 	{
-		EFilterResult::Type Result = ChildFilter->IsPropertyValid(Params);
-		const bool bShouldNegate = EditorFilterBehavior == EEditorFilterBehavior::Negate;
-		return bShouldNegate ? EFilterResult::Negate(Result) : Result;
-	}
-	return EFilterResult::DoNotCare;
+		// Child filter is nullptr when user deletes filter class and force deletes
+		return ChildFilter ? ChildFilter->IsAddedActorValid(Params) : EFilterResult::DoNotCare;
+	});
 }
diff --git a/Engine/Plugins/VirtualProduction/LevelSnapshots/Source/LevelSnapshotsEditor/Private/Data/Filters/NegatableFilter.h b/Engine/Plugins/VirtualProduction/LevelSnapshots/Source/LevelSnapshotsEditor/Private/Data/Filters/NegatableFilter.h
--- a/Engine/Plugins/VirtualProduction/LevelSnapshots/Source/LevelSnapshotsEditor/Private/Data/Filters/NegatableFilter.h
+++ b/Engine/Plugins/VirtualProduction/LevelSnapshots/Source/LevelSnapshotsEditor/Private/Data/Filters/NegatableFilter.h
@@ -22,16 +22,21 @@ public:
 	FText GetDisplayName() const;
 
 	//~ Begin ULevelSnapshotFilter Interface
-	EFilterResult::Type IsActorValid(const FIsActorValidParams& Params) const override;
-	EFilterResult::Type IsPropertyValid(const FIsPropertyValidParams& Params) const override;
+	virtual EFilterResult::Type IsActorValid(const FIsActorValidParams& Params) const override;
+	virtual EFilterResult::Type IsPropertyValid(const FIsPropertyValidParams& Params) const override;
+	virtual EFilterResult::Type IsDeletedActorValid(const FIsDeletedActorValidParams& Params) const override;
+	virtual EFilterResult::Type IsAddedActorValid(const FIsAddedActorValidParams& Params) const override;
 	//~ End ULevelSnapshotFilter Interface
+
 	
 public: // Only public to use GET_MEMBER_NAME_CHECKED with compiler checks - do not use directly.
+	
 	/* Display name in editor. Defaults to class name if left empty. */
 	UPROPERTY(EditAnywhere, Category = "Filter")
 	FString Name;
+
+private:
 	
 	UPROPERTY()
 	ULevelSnapshotFilter* ChildFilter;
-	
 };
\ No newline at end of file
diff --git a/Engine/Plugins/VirtualProduction/LevelSnapshots/Source/LevelSnapshotsEditor/Private/LevelSnapshotsEditorModule.cpp b/Engine/Plugins/VirtualProduction/LevelSnapshots/Source/LevelSnapshotsEditor/Private/LevelSnapshotsEditorModule.cpp
--- a/Engine/Plugins/VirtualProduction/LevelSnapshots/Source/LevelSnapshotsEditor/Private/LevelSnapshotsEditorModule.cpp
+++ b/Engine/Plugins/VirtualProduction/LevelSnapshots/Source/LevelSnapshotsEditor/Private/LevelSnapshotsEditorModule.cpp
@@ -4,25 +4,23 @@
 
 #include "Settings/LevelSnapshotsEditorProjectSettings.h"
 #include "NegatableFilter.h"
-#include "NegatableFilterDetailsCustomization.h"
 #include "LevelSnapshotsEditorCommands.h"
 #include "LevelSnapshotsEditorStyle.h"
+#include "LevelSnapshotsEditorFunctionLibrary.h"
 #include "LevelSnapshotsEditorData.h"
 #include "LevelSnapshotsFunctionLibrary.h"
 #include "LevelSnapshotsUserSettings.h"
 #include "SLevelSnapshotsEditorCreationForm.h"
 #include "Toolkits/LevelSnapshotsEditorToolkit.h"
 
+#include "AssetRegistry/AssetRegistryModule.h"
 #include "AssetToolsModule.h"
 #include "AssetTypeActions/AssetTypeActions_LevelSnapshot.h"
 #include "IAssetTools.h"
 #include "ISettingsModule.h"
 #include "LevelEditor.h"
-#include "LevelSnapshotsEditorFunctionLibrary.h"
-#include "Editor/MainFrame/Private/Menus/SettingsMenu.h"
 #include "ToolMenus.h"
 #include "ToolMenuSection.h"
-#include "AssetRegistry/AssetRegistryModule.h"
 #include "Framework/Notifications/NotificationManager.h"
 #include "Misc/ScopeExit.h"
 #include "Widgets/Notifications/SNotificationList.h"
@@ -49,12 +47,6 @@ void FLevelSnapshotsEditorModule::StartupModule()
 	{
 		RegisterEditorToolbar();
 	}
-
-	FPropertyEditorModule& PropertyModule = FModuleManager::LoadModuleChecked<FPropertyEditorModule>("PropertyEditor");
-	PropertyModule.RegisterCustomClassLayout( UNegatableFilter::StaticClass()->GetFName(), FOnGetDetailCustomizationInstance::CreateLambda( []()
-	{
-		return MakeShared<FNegatableFilterDetailsCustomization>();
-	}));
 }
 
 void FLevelSnapshotsEditorModule::ShutdownModule()
@@ -301,6 +293,7 @@ void FLevelSnapshotsEditorModule::TakeAndSaveSnapshot(const FText& InDescription
 	{
 		return;
 	}
+	Snapshot->SetFlags(RF_Public | RF_Standalone);
 
 	
 	// Take screenshot before we save
@@ -345,7 +338,7 @@ void FLevelSnapshotsEditorModule::OpenSnapshotsEditor()
 
 void FLevelSnapshotsEditorModule::OpenLevelSnapshotsSettings()
 {
-	FSettingsMenu::OpenSettings("Project", "Plugins", "Level Snapshots");
+	FModuleManager::LoadModuleChecked<ISettingsModule>("Settings").ShowViewer("Project", "Plugins", "Level Snapshots");
 }
 
 ULevelSnapshotsEditorData* FLevelSnapshotsEditorModule::AllocateTransientPreset()
diff --git a/Engine/Plugins/VirtualProduction/LevelSnapshots/Source/LevelSnapshotsEditor/Private/Views/Filter/SLevelSnapshotsEditorFilters.cpp b/Engine/Plugins/VirtualProduction/LevelSnapshots/Source/LevelSnapshotsEditor/Private/Views/Filter/SLevelSnapshotsEditorFilters.cpp
--- a/Engine/Plugins/VirtualProduction/LevelSnapshots/Source/LevelSnapshotsEditor/Private/Views/Filter/SLevelSnapshotsEditorFilters.cpp
+++ b/Engine/Plugins/VirtualProduction/LevelSnapshots/Source/LevelSnapshotsEditor/Private/Views/Filter/SLevelSnapshotsEditorFilters.cpp
@@ -18,6 +18,7 @@
 #include "EditorStyleSet.h"
 #include "Framework/Notifications/NotificationManager.h"
 #include "IDetailsView.h"
+#include "NegatableFilterDetailsCustomization.h"
 #include "Modules/ModuleManager.h"
 #include "PropertyEditorModule.h"
 #include "Widgets/Input/SButton.h"
@@ -129,8 +130,11 @@ void SLevelSnapshotsEditorFilters::Construct(const FArguments& InArgs, const TSh
 	DetailsViewArgs.NameAreaSettings = FDetailsViewArgs::HideNameArea;
 	DetailsViewArgs.bHideSelectionTip = true;
 	DetailsViewArgs.DefaultsOnlyVisibility = EEditDefaultsOnlyNodeVisibility::Automatic;
-
+	
 	FilterDetailsView = EditModule.CreateDetailView(DetailsViewArgs);
+	FilterDetailsView->RegisterInstancedCustomPropertyLayout(UNegatableFilter::StaticClass(),
+		FOnGetDetailCustomizationInstance::CreateLambda([]() { return MakeShared<FNegatableFilterDetailsCustomization>(); })
+		);
 
 
 	
diff --git a/Engine/Plugins/VirtualProduction/LevelSnapshots/Source/LevelSnapshotsEditor/Private/Views/Results/SLevelSnapshotsEditorResults.cpp b/Engine/Plugins/VirtualProduction/LevelSnapshots/Source/LevelSnapshotsEditor/Private/Views/Results/SLevelSnapshotsEditorResults.cpp
--- a/Engine/Plugins/VirtualProduction/LevelSnapshots/Source/LevelSnapshotsEditor/Private/Views/Results/SLevelSnapshotsEditorResults.cpp
+++ b/Engine/Plugins/VirtualProduction/LevelSnapshots/Source/LevelSnapshotsEditor/Private/Views/Results/SLevelSnapshotsEditorResults.cpp
@@ -152,6 +152,9 @@ void FLevelSnapshotsEditorResultsRow::InitHeaderRow(const ELevelSnapshotsEditorR
 	HeaderType = InHeaderType;
 	HeaderColumns = InColumns;
 	ResultsViewPtr = InResultsView;
+
+	// Set visibility values to false for header rows so they hide correctly when they have no visible children
+	bDoesRowMatchSearchTerms = false;
 }
 
 void FLevelSnapshotsEditorResultsRow::InitAddedActorRow(AActor* InAddedActor)
@@ -741,6 +744,17 @@ void FLevelSnapshotsEditorResultsRow::GenerateActorGroupChildren(FPropertySelect
 	}
 
 	SetHasGeneratedChildren(true);
+
+	// Remove cached search terms now that the actor group has child rows to search
+	SetCachedSearchTerms("");
+
+	// Apply Search
+	check(ResultsViewPtr.IsValid());
+	
+	if (const TSharedPtr<SLevelSnapshotsEditorResults>& PinnedResults = ResultsViewPtr.Pin())
+	{
+		PinnedResults->ExecuteResultsViewSearchOnSpecifiedActors(PinnedResults->GetSearchStringFromSearchInputField(), { SharedThis(this) });
+	}
 }
 
 bool FLevelSnapshotsEditorResultsRow::DoesRowRepresentGroup() const
@@ -937,7 +951,7 @@ bool FLevelSnapshotsEditorResultsRow::MatchSearchTokensToSearchTerms(const TArra
 	}
 	else
 	{
-		const FString& SearchTerms = GetSearchTerms();
+		const FString& SearchTerms = GetOrCacheSearchTerms();
 
 		for (const FString& Token : InTokens)
 		{
@@ -1203,7 +1217,7 @@ void FLevelSnapshotsEditorResultsRow::SetIsNodeChecked(const bool bNewChecked)
 }
 
 bool FLevelSnapshotsEditorResultsRow::HasVisibleChildren() const
-{
+{	
 	bool bVisibleChildFound = false;
 
 	for (const TSharedPtr<FLevelSnapshotsEditorResultsRow>& ChildRow : GetChildRows())
@@ -1287,7 +1301,7 @@ bool FLevelSnapshotsEditorResultsRow::HasChangedChildren() const
 
 bool FLevelSnapshotsEditorResultsRow::ShouldRowBeVisible() const
 {
-	return GetDoesRowMatchSearchTerms();
+	return GetDoesRowMatchSearchTerms() || HasVisibleChildren();
 }
 
 EVisibility FLevelSnapshotsEditorResultsRow::GetDesiredVisibility() const
@@ -1295,9 +1309,19 @@ EVisibility FLevelSnapshotsEditorResultsRow::GetDesiredVisibility() const
 	return ShouldRowBeVisible() ? EVisibility::Visible : EVisibility::Collapsed;
 }
 
-const FString& FLevelSnapshotsEditorResultsRow::GetSearchTerms() const
+const FString& FLevelSnapshotsEditorResultsRow::GetOrCacheSearchTerms()
+{
+	if (CachedSearchTerms.IsEmpty())
+	{
+		SetCachedSearchTerms(GetDisplayName().ToString());
+	}
+	
+	return CachedSearchTerms;
+}
+
+void FLevelSnapshotsEditorResultsRow::SetCachedSearchTerms(const FString& InTerms)
 {
-	return GetDisplayName().ToString();
+	CachedSearchTerms = InTerms;
 }
 
 bool FLevelSnapshotsEditorResultsRow::GetDoesRowMatchSearchTerms() const
@@ -1448,7 +1472,7 @@ void SLevelSnapshotsEditorResults::Construct(const FArguments& InArgs, ULevelSna
 				.VAlign(VAlign_Center)
 				[
 					SAssignNew(SelectedSnapshotNamePtr, STextBlock)
-					.Text(NSLOCTEXT("LevelSnapshots", "LevelSnapshots", "Level Snapshots"))
+					.Text(LOCTEXT("LevelSnapshots", "Level Snapshots"))
 					.Font(FCoreStyle::GetDefaultFontStyle("Bold", 16))
 				]
 			]	
@@ -1489,21 +1513,44 @@ void SLevelSnapshotsEditorResults::Construct(const FArguments& InArgs, ULevelSna
 
 			+ SVerticalBox::Slot()
 			[
-				SAssignNew(TreeViewPtr, STreeView<FLevelSnapshotsEditorResultsRowPtr>)
-				.SelectionMode(ESelectionMode::None)
-				.TreeItemsSource(&TreeViewRootHeaderObjects)
-				.OnGenerateRow_Lambda([this](FLevelSnapshotsEditorResultsRowPtr Row, const TSharedRef<STableViewBase>& OwnerTable)
-					{
-						check(Row.IsValid());
-					
-						return SNew(STableRow<FLevelSnapshotsEditorResultsRowPtr>, OwnerTable)
-							[
-								SNew(SLevelSnapshotsEditorResultsRow, Row, SplitterManagerPtr)
-							]
-							.Visibility_Raw(Row.Get(), &FLevelSnapshotsEditorResultsRow::GetDesiredVisibility);
-					})
-				.OnGetChildren_Raw(this, &SLevelSnapshotsEditorResults::OnGetRowChildren)
-				.OnExpansionChanged_Raw(this, &SLevelSnapshotsEditorResults::OnRowChildExpansionChange)
+				SNew(SOverlay)
+				
+				+ SOverlay::Slot()
+				.HAlign(HAlign_Fill)
+				.Padding(2.0f, 2.0f, 2.0f, 2.0f)
+				[
+					SAssignNew(TreeViewPtr, STreeView<FLevelSnapshotsEditorResultsRowPtr>)
+					.SelectionMode(ESelectionMode::None)
+					.TreeItemsSource(&TreeViewRootHeaderObjects)
+					.OnGenerateRow_Lambda([this](FLevelSnapshotsEditorResultsRowPtr Row, const TSharedRef<STableViewBase>& OwnerTable)
+						{
+							check(Row.IsValid());
+						
+							return SNew(STableRow<FLevelSnapshotsEditorResultsRowPtr>, OwnerTable)
+								[
+									SNew(SLevelSnapshotsEditorResultsRow, Row, SplitterManagerPtr)
+								]
+								.Visibility_Raw(Row.Get(), &FLevelSnapshotsEditorResultsRow::GetDesiredVisibility);
+						})
+					.OnGetChildren_Raw(this, &SLevelSnapshotsEditorResults::OnGetRowChildren)
+					.OnExpansionChanged_Raw(this, &SLevelSnapshotsEditorResults::OnRowChildExpansionChange)
+					.Visibility_Lambda([this]()
+						{
+							return this->DoesTreeViewHaveVisibleChildren() ? EVisibility::Visible : EVisibility::Collapsed;
+						})
+				]
+
+				+ SOverlay::Slot()
+				.HAlign(HAlign_Center)
+				.Padding(2.0f, 24.0f, 2.0f, 2.0f)
+				[
+					SNew(STextBlock)
+					.Text(LOCTEXT("LevelSnapshotsEditorResults_NoResults", "No results to show. Try selecting a snapshot, changing your active filters, or clearing any active search."))
+					.Visibility_Lambda([this]()
+						{
+							return this->DoesTreeViewHaveVisibleChildren() ? EVisibility::Collapsed : EVisibility::HitTestInvisible;
+						})
+				]
 			]
 
 			+SVerticalBox::Slot()
@@ -1931,9 +1978,10 @@ void SLevelSnapshotsEditorResults::BuildSelectionSetFromSelectedPropertiesInEach
 	GetEditorDataPtr()->GetFilterResults()->SetPropertiesToRollback(PropertySelectionMap);
 }
 
-const FString& SLevelSnapshotsEditorResults::GetSearchStringFromSearchInputField() const
+FString SLevelSnapshotsEditorResults::GetSearchStringFromSearchInputField() const
 {
-	return ResultsSearchBoxPtr->GetText().ToString();
+	return ensureAlwaysMsgf(ResultsSearchBoxPtr.IsValid(), TEXT("%hs: ResultsSearchBoxPtr is not valid. Check to make sure it was created."), __FUNCTION__)
+	? ResultsSearchBoxPtr->GetText().ToString() : "";
 }
 
 ULevelSnapshotsEditorData* SLevelSnapshotsEditorResults::GetEditorDataPtr() const
@@ -2078,6 +2126,9 @@ void SLevelSnapshotsEditorResults::GenerateTreeView()
 
 	TreeViewPtr->RequestListRefresh();
 	UpdateSnapshotInformationText();
+
+	// Apply last search
+	ExecuteResultsViewSearchOnAllActors(GetSearchStringFromSearchInputField());
 }
 
 bool SLevelSnapshotsEditorResults::GenerateTreeViewChildren_ModifiedActors(FLevelSnapshotsEditorResultsRowPtr ModifiedActorsHeader, ULevelSnapshotFilter* UserFilters)
@@ -2115,7 +2166,8 @@ bool SLevelSnapshotsEditorResults::GenerateTreeViewChildren_ModifiedActors(FLeve
 			FilterListData.ApplyFilterToFindSelectedProperties(WorldActor, UserFilters);
 		}
 
-		const int32 KeyCountAfterFilter = FilterListData.GetModifiedActorsSelectedProperties().GetKeyCount();
+		const FPropertySelectionMap& ModifiedSelectedActors = FilterListData.GetModifiedActorsSelectedProperties();
+		const int32 KeyCountAfterFilter = ModifiedSelectedActors.GetKeyCount();
 		const int32 KeyCountDifference = KeyCountAfterFilter - KeyCountBeforeFilter;
 
 		// If keys have been added, then this actor or its children have properties that pass the filter
@@ -2133,12 +2185,37 @@ bool SLevelSnapshotsEditorResults::GenerateTreeViewChildren_ModifiedActors(FLeve
 		// Create group
 		FLevelSnapshotsEditorResultsRowPtr NewActorGroup = MakeShared<FLevelSnapshotsEditorResultsRow>(
 			FLevelSnapshotsEditorResultsRow(FText::FromString(ActorName), FLevelSnapshotsEditorResultsRow::ActorGroup, ECheckBoxState::Checked, ModifiedActorsHeader));
-		NewActorGroup->InitActorRow(WeakSnapshotActor.IsValid() ? WeakSnapshotActor.Get() : nullptr, 
-			WeakWorldActor.IsValid() ? WeakWorldActor.Get() : nullptr, SharedThis(this));
+		
+		NewActorGroup->InitActorRow(WeakSnapshotActor.IsValid() ? WeakSnapshotActor.Get() : nullptr, WorldActor, SharedThis(this));
 
 		ModifiedActorsHeader->AddToChildRows(NewActorGroup);
 		
 		TreeViewModifiedActorGroupObjects.Add(NewActorGroup);
+
+		// Cache search terms using the desired leaf properties for each object newly added to ModifiedActorsSelectedProperties this loop
+		FString NewCachedSearchTerms = WorldActor->GetHumanReadableName();
+
+		if (const FPropertySelection* ActorProperties = ModifiedSelectedActors.GetSelectedProperties(WorldActor))
+		{
+			for (const TFieldPath<FProperty>& LeafProperty : ActorProperties->GetSelectedLeafProperties())
+			{
+				NewCachedSearchTerms += " " + LeafProperty.ToString();
+			}
+		}
+
+		for (UActorComponent* Component : WorldActor->GetComponents())
+		{
+			if (const FPropertySelection* ComponentProperties = ModifiedSelectedActors.GetSelectedProperties(Component))
+			{
+				NewCachedSearchTerms += " " + Component->GetReadableName();
+				for (const TFieldPath<FProperty>& LeafProperty : ComponentProperties->GetSelectedLeafProperties())
+				{
+					NewCachedSearchTerms += " " + LeafProperty.ToString();
+				}
+			}
+		}
+		
+		NewActorGroup->SetCachedSearchTerms(NewCachedSearchTerms);
 	}
 
 	return TreeViewModifiedActorGroupObjects.Num() > 0;
@@ -2213,19 +2290,28 @@ FReply SLevelSnapshotsEditorResults::SetAllActorGroupsCollapsed()
 	return FReply::Handled();
 }
 
-void SLevelSnapshotsEditorResults::OnResultsViewSearchTextChanged(const FText& Text)
+void SLevelSnapshotsEditorResults::OnResultsViewSearchTextChanged(const FText& Text) const
 {
-	ExecuteResultsViewSearch(Text.ToString());
+	ExecuteResultsViewSearchOnAllActors(Text.ToString());
 }
 
-void SLevelSnapshotsEditorResults::ExecuteResultsViewSearch(const FString& SearchString)
+void SLevelSnapshotsEditorResults::ExecuteResultsViewSearchOnAllActors(const FString& SearchString) const
+{
+	// Consider all rows for search except the header rows
+	ExecuteResultsViewSearchOnSpecifiedActors(SearchString, TreeViewModifiedActorGroupObjects);
+	ExecuteResultsViewSearchOnSpecifiedActors(SearchString, TreeViewAddedActorGroupObjects);
+	ExecuteResultsViewSearchOnSpecifiedActors(SearchString, TreeViewRemovedActorGroupObjects);
+}
+
+void SLevelSnapshotsEditorResults::ExecuteResultsViewSearchOnSpecifiedActors(
+	const FString& SearchString, const TArray<TSharedPtr<FLevelSnapshotsEditorResultsRow>>& ActorRowsToConsider) const
 {
 	TArray<FString> Tokens;
 	
 	// unquoted search equivalent to a match-any-of search
 	SearchString.ParseIntoArray(Tokens, TEXT(" "), true);
 	
-	for (const TSharedPtr<FLevelSnapshotsEditorResultsRow>& ChildRow : TreeViewModifiedActorGroupObjects)
+	for (const TSharedPtr<FLevelSnapshotsEditorResultsRow>& ChildRow : ActorRowsToConsider)
 	{
 		if (!ensure(ChildRow.IsValid()))
 		{
@@ -2240,6 +2326,24 @@ void SLevelSnapshotsEditorResults::ExecuteResultsViewSearch(const FString& Searc
 	}
 }
 
+bool SLevelSnapshotsEditorResults::DoesTreeViewHaveVisibleChildren() const
+{
+	if (TreeViewPtr.IsValid())
+	{
+		for (const TSharedPtr<FLevelSnapshotsEditorResultsRow>& Header : TreeViewRootHeaderObjects)
+		{
+			const EVisibility HeaderVisibility = Header->GetDesiredVisibility();
+			
+			if (HeaderVisibility != EVisibility::Hidden && HeaderVisibility != EVisibility::Collapsed)
+			{
+				return true;
+			}
+		}
+	}
+	
+	return false;
+}
+
 void SLevelSnapshotsEditorResults::OnGetRowChildren(FLevelSnapshotsEditorResultsRowPtr Row, TArray<FLevelSnapshotsEditorResultsRowPtr>& OutChildren)
 {
 	if (Row.IsValid())
@@ -2265,7 +2369,7 @@ void SLevelSnapshotsEditorResults::OnGetRowChildren(FLevelSnapshotsEditorResults
 	}
 }
 
-void SLevelSnapshotsEditorResults::OnRowChildExpansionChange(FLevelSnapshotsEditorResultsRowPtr Row, const bool bIsExpanded)
+void SLevelSnapshotsEditorResults::OnRowChildExpansionChange(FLevelSnapshotsEditorResultsRowPtr Row, const bool bIsExpanded) const
 {
 	if (Row.IsValid())
 	{
diff --git a/Engine/Plugins/VirtualProduction/LevelSnapshots/Source/LevelSnapshotsEditor/Private/Views/Results/SLevelSnapshotsEditorResults.h b/Engine/Plugins/VirtualProduction/LevelSnapshots/Source/LevelSnapshotsEditor/Private/Views/Results/SLevelSnapshotsEditorResults.h
--- a/Engine/Plugins/VirtualProduction/LevelSnapshots/Source/LevelSnapshotsEditor/Private/Views/Results/SLevelSnapshotsEditorResults.h
+++ b/Engine/Plugins/VirtualProduction/LevelSnapshots/Source/LevelSnapshotsEditor/Private/Views/Results/SLevelSnapshotsEditorResults.h
@@ -93,14 +93,14 @@ struct FLevelSnapshotsEditorResultsRow final : TSharedFromThis<FLevelSnapshotsEd
 	enum ELevelSnapshotsEditorResultsRowType
 	{
 		None,
-		TreeViewHeader,
-		AddedActor,
-		RemovedActor,
-		ActorGroup,
-		ComponentGroup,
-		SubObjectGroup,
+		TreeViewHeader, // Includes rows that represent headers grouping together added, removed, or modified actor rows
+		AddedActor, // Includes rows that represent actors which have been created in the world since the snapshot was taken
+		RemovedActor, // Includes rows that represent actors which have been deleted from the world since the snapshot was taken
+		ActorGroup, // Modified Actor group. Use AddedActor or RemovedActor for other Actor row types.
+		ComponentGroup, // Includes rows that represent UActorComponent types
+		SubObjectGroup, // Includes rows that represent separate non-component objects wholly owned by the world or snapshot actor
 		StructGroup,
-		CollectionGroup,
+		CollectionGroup, // Includes rows that represent TMap, TSet, and TArray.
 		SingleProperty,
 		SinglePropertyInStruct,
 		SinglePropertyInMap,
@@ -184,6 +184,8 @@ struct FLevelSnapshotsEditorResultsRow final : TSharedFromThis<FLevelSnapshotsEd
 	void ExecuteSearchOnChildNodes(const FString& SearchString) const;
 	void ExecuteSearchOnChildNodes(const TArray<FString>& Tokens) const;
 
+	void SetCachedSearchTerms(const FString& InTerms);
+
 	UObject* GetSnapshotObject() const;
 	UObject* GetWorldObject() const;
 
@@ -253,7 +255,7 @@ private:
 	bool bDoesRowMatchSearchTerms = true;
 
 	/* Returns a string of searchable keywords such as object names, property names, paths or anything else associated with the row that might be useful to search for. */
-	const FString& GetSearchTerms() const;
+	const FString& GetOrCacheSearchTerms();
 
 	/* Use MatchSearchTerms() first to match Search Terms against tokens, then call this method*/
 	bool GetDoesRowMatchSearchTerms() const;
@@ -323,7 +325,11 @@ public:
 	/* This method builds a selection set of all visible and checked properties to pass back to apply to the world. */
 	void BuildSelectionSetFromSelectedPropertiesInEachActorGroup();
 
-	const FString& GetSearchStringFromSearchInputField() const;
+	FString GetSearchStringFromSearchInputField() const;
+	void ExecuteResultsViewSearchOnAllActors(const FString& SearchString) const;
+	void ExecuteResultsViewSearchOnSpecifiedActors(const FString& SearchString, const TArray<TSharedPtr<FLevelSnapshotsEditorResultsRow>>& ActorRowsToConsider) const;
+
+	bool DoesTreeViewHaveVisibleChildren() const;
 
 	ULevelSnapshotsEditorData* GetEditorDataPtr() const;
 
@@ -355,8 +361,7 @@ private:
 
 	// Search
 	
-	void OnResultsViewSearchTextChanged(const FText& Text);
-	void ExecuteResultsViewSearch(const FString& SearchString);
+	void OnResultsViewSearchTextChanged(const FText& Text) const;
 
 	FDelegateHandle OnActiveSnapshotChangedHandle;
 	FDelegateHandle OnRefreshResultsHandle;
@@ -382,7 +387,7 @@ private:
 	bool GenerateTreeViewChildren_RemovedActors(FLevelSnapshotsEditorResultsRowPtr RemovedActorsHeader);
 	
 	void OnGetRowChildren(FLevelSnapshotsEditorResultsRowPtr Row, TArray<FLevelSnapshotsEditorResultsRowPtr>& OutChildren);
-	void OnRowChildExpansionChange(FLevelSnapshotsEditorResultsRowPtr Row, const bool bIsExpanded);
+	void OnRowChildExpansionChange(FLevelSnapshotsEditorResultsRowPtr Row, const bool bIsExpanded) const;
 	
 	TSharedPtr<STreeView<FLevelSnapshotsEditorResultsRowPtr>> TreeViewPtr;
 	
