Fix Horde issue.

diff --git a/Engine/Plugins/VirtualProduction/LevelSnapshots/Source/LevelSnapshots/Private/Archive/SnapshotArchive.cpp b/Engine/Plugins/VirtualProduction/LevelSnapshots/Source/LevelSnapshots/Private/Archive/SnapshotArchive.cpp
--- a/Engine/Plugins/VirtualProduction/LevelSnapshots/Source/LevelSnapshots/Private/Archive/SnapshotArchive.cpp
+++ b/Engine/Plugins/VirtualProduction/LevelSnapshots/Source/LevelSnapshots/Private/Archive/SnapshotArchive.cpp
@@ -166,8 +166,8 @@ FSnapshotArchive::FSnapshotArchive(FObjectSnapshotData& InObjectData, FWorldSnap
 	{
 		const FSnapshotVersionInfo& VersionInfo = InSharedData.GetSnapshotVersionInfo();
 		
-		Super::SetUE4Ver(VersionInfo.FileVersion.FileVersionUE4);
-		Super::SetLicenseeUE4Ver(VersionInfo.FileVersion.FileVersionLicenseeUE4);
+		Super::SetUEVer(VersionInfo.FileVersion.FileVersion);
+		Super::SetLicenseeUEVer(VersionInfo.FileVersion.FileVersionLicensee);
 		Super::SetEngineVer(FEngineVersionBase(VersionInfo.EngineVersion.Major, VersionInfo.EngineVersion.Minor, VersionInfo.EngineVersion.Patch, VersionInfo.EngineVersion.Changelist));
 
 		FCustomVersionContainer EngineCustomVersions;
diff --git a/Engine/Plugins/VirtualProduction/LevelSnapshots/Source/LevelSnapshots/Private/Data/SnapshotVersion.cpp b/Engine/Plugins/VirtualProduction/LevelSnapshots/Source/LevelSnapshots/Private/Data/SnapshotVersion.cpp
--- a/Engine/Plugins/VirtualProduction/LevelSnapshots/Source/LevelSnapshots/Private/Data/SnapshotVersion.cpp
+++ b/Engine/Plugins/VirtualProduction/LevelSnapshots/Source/LevelSnapshots/Private/Data/SnapshotVersion.cpp
@@ -3,10 +3,11 @@
 #include "Data/SnapshotVersion.h"
 #include "Misc/EngineVersion.h"
 #include "Serialization/CustomVersion.h"
+
 void FSnapshotFileVersionInfo::Initialize()
 {
-	FileVersionUE4 = GPackageFileUE4Version;
-	FileVersionLicenseeUE4 = GPackageFileLicenseeUE4Version;
+	FileVersion = GPackageFileUEVersion;
+	FileVersionLicensee = GPackageFileLicenseeUEVersion;
 }
 
 void FSnapshotEngineVersionInfo::Initialize(const FEngineVersion& InVersion)
diff --git a/Engine/Plugins/VirtualProduction/LevelSnapshots/Source/LevelSnapshots/Public/Data/SnapshotVersion.h b/Engine/Plugins/VirtualProduction/LevelSnapshots/Source/LevelSnapshots/Public/Data/SnapshotVersion.h
--- a/Engine/Plugins/VirtualProduction/LevelSnapshots/Source/LevelSnapshots/Public/Data/SnapshotVersion.h
+++ b/Engine/Plugins/VirtualProduction/LevelSnapshots/Source/LevelSnapshots/Public/Data/SnapshotVersion.h
@@ -18,13 +18,13 @@ struct LEVELSNAPSHOTS_API FSnapshotFileVersionInfo
 	/** Initialize this version info from the compiled in data */
 	void Initialize();
 
-	/* UE4 file version */
+	/* UE File version */
 	UPROPERTY()
-	int32 FileVersionUE4 = 0;
+	int32 FileVersion = 0;
 
 	/* Licensee file version */
 	UPROPERTY()
-	int32 FileVersionLicenseeUE4 = 0;
+	int32 FileVersionLicensee = 0;
 };
 
 /** Holds engine version information */
