Remove legacy data

diff --git a/Engine/Plugins/VirtualProduction/LevelSnapshots/Source/LevelSnapshots/Private/Archive/ApplySnapshotDataArchive.cpp b/Engine/Plugins/VirtualProduction/LevelSnapshots/Source/LevelSnapshots/Private/Archive/ApplySnapshotDataArchive.cpp
deleted file mode 100644
diff --git a/Engine/Plugins/VirtualProduction/LevelSnapshots/Source/LevelSnapshots/Private/Archive/ApplySnapshotDataArchive.h b/Engine/Plugins/VirtualProduction/LevelSnapshots/Source/LevelSnapshots/Private/Archive/ApplySnapshotDataArchive.h
deleted file mode 100644
diff --git a/Engine/Plugins/VirtualProduction/LevelSnapshots/Source/LevelSnapshots/Private/Archive/SnapshotArchives.cpp b/Engine/Plugins/VirtualProduction/LevelSnapshots/Source/LevelSnapshots/Private/Archive/SnapshotArchives.cpp
deleted file mode 100644
diff --git a/Engine/Plugins/VirtualProduction/LevelSnapshots/Source/LevelSnapshots/Private/Archive/SnapshotArchives.h b/Engine/Plugins/VirtualProduction/LevelSnapshots/Source/LevelSnapshots/Private/Archive/SnapshotArchives.h
deleted file mode 100644
diff --git a/Engine/Plugins/VirtualProduction/LevelSnapshots/Source/LevelSnapshots/Private/Archive/TakeSnapshotArchive.cpp b/Engine/Plugins/VirtualProduction/LevelSnapshots/Source/LevelSnapshots/Private/Archive/TakeSnapshotArchive.cpp
deleted file mode 100644
diff --git a/Engine/Plugins/VirtualProduction/LevelSnapshots/Source/LevelSnapshots/Private/Archive/TakeSnapshotArchive.h b/Engine/Plugins/VirtualProduction/LevelSnapshots/Source/LevelSnapshots/Private/Archive/TakeSnapshotArchive.h
deleted file mode 100644
diff --git a/Engine/Plugins/VirtualProduction/LevelSnapshots/Source/LevelSnapshots/Private/Data/ActorSnapshot.cpp b/Engine/Plugins/VirtualProduction/LevelSnapshots/Source/LevelSnapshots/Private/Data/ActorSnapshot.cpp
deleted file mode 100644
diff --git a/Engine/Plugins/VirtualProduction/LevelSnapshots/Source/LevelSnapshots/Private/Data/ActorSnapshotData.cpp b/Engine/Plugins/VirtualProduction/LevelSnapshots/Source/LevelSnapshots/Private/Data/ActorSnapshotData.cpp
--- a/Engine/Plugins/VirtualProduction/LevelSnapshots/Source/LevelSnapshots/Private/Data/ActorSnapshotData.cpp
+++ b/Engine/Plugins/VirtualProduction/LevelSnapshots/Source/LevelSnapshots/Private/Data/ActorSnapshotData.cpp
@@ -3,7 +3,6 @@
 #include "ActorSnapshotData.h"
 
 #include "ApplySnapshotDataArchiveV2.h"
-#include "LevelSnapshotSelections.h"
 #include "LevelSnapshotsLog.h"
 #include "PropertySelectionMap.h"
 #include "TakeWorldObjectSnapshotArchive.h"
diff --git a/Engine/Plugins/VirtualProduction/LevelSnapshots/Source/LevelSnapshots/Private/Data/BaseObjectInfo.cpp b/Engine/Plugins/VirtualProduction/LevelSnapshots/Source/LevelSnapshots/Private/Data/BaseObjectInfo.cpp
deleted file mode 100644
diff --git a/Engine/Plugins/VirtualProduction/LevelSnapshots/Source/LevelSnapshots/Private/Data/LevelSnapshot.cpp b/Engine/Plugins/VirtualProduction/LevelSnapshots/Source/LevelSnapshots/Private/Data/LevelSnapshot.cpp
--- a/Engine/Plugins/VirtualProduction/LevelSnapshots/Source/LevelSnapshots/Private/Data/LevelSnapshot.cpp
+++ b/Engine/Plugins/VirtualProduction/LevelSnapshots/Source/LevelSnapshots/Private/Data/LevelSnapshot.cpp
@@ -2,7 +2,6 @@
 
 #include "Data/LevelSnapshot.h"
 
-#include "LevelSnapshotSelections.h"
 #include "LevelSnapshotsLog.h"
 #include "LevelSnapshotsModule.h"
 #include "LevelSnapshotsStats.h"
@@ -107,17 +106,7 @@ void ULevelSnapshot::ApplySnapshotToWorld(UWorld* TargetWorld, const FPropertySe
 #if WITH_EDITOR
 	FScopedTransaction Transaction(FText::FromString("Loading Level Snapshot."));
 #endif
-	
-	// Temporary fix until asset migration is implemented: simply use old system for old data.
-	const bool bHasLegacyData = ActorSnapshots.Num() > 0;
-	if (bHasLegacyData)
-	{
-		LegacyApplySnapshotToWorld(SelectionSet);
-	}
-	else
-	{
-		SerializedData.ApplyToWorld(TargetWorld, GetPackage(), SelectionSet);
-	}
+	SerializedData.ApplyToWorld(TargetWorld, GetPackage(), SelectionSet);
 }
 
 void ULevelSnapshot::SnapshotWorld(UWorld* TargetWorld)
@@ -142,12 +131,6 @@ void ULevelSnapshot::SnapshotWorld(UWorld* TargetWorld)
 		UE_LOG(LogLevelSnapshots, Warning, TEXT("Level snapshots currently only support editors. Snapshots taken in other world types are experimental any may not function as expected."));
 	}
 	
-	const bool bHasLegacyData = ActorSnapshots.Num() > 0;
-	if (bHasLegacyData)
-	{
-		ActorSnapshots.Empty();
-	}
-	
 	EnsureWorldInitialised();
 	MapPath = TargetWorld;
 	CaptureTime = FDateTime::UtcNow();
@@ -297,25 +280,14 @@ bool ULevelSnapshot::AreSnapshotAndOriginalPropertiesEquivalent(const FProperty*
 TOptional<AActor*> ULevelSnapshot::GetDeserializedActor(const FSoftObjectPath& OriginalActorPath)
 {
 	DECLARE_SCOPE_CYCLE_COUNTER(TEXT("GetDeserializedActor"), STAT_GetDeserializedActor, STATGROUP_LevelSnapshots);
+	
 	EnsureWorldInitialised();
-
-	// Temporary fix until asset migration is implemented: simply use old system for old data.
-	const bool bHasLegacyData = ActorSnapshots.Num() > 0;
-	if (bHasLegacyData)
-	{
-		FLevelSnapshot_Actor* LegacySnapshotData = ActorSnapshots.Find(OriginalActorPath);
-		return LegacySnapshotData ? LegacySnapshotData->GetDeserializedActor(SnapshotContainerWorld) : TOptional<AActor*>();
-	}
-	else
-	{
-		return SerializedData.GetDeserializedActor(OriginalActorPath, GetPackage());
-	}
+	return SerializedData.GetDeserializedActor(OriginalActorPath, GetPackage());
 }
 
 int32 ULevelSnapshot::GetNumSavedActors() const
 {
-	const bool bHasLegacyData = ActorSnapshots.Num() > 0;
-	return bHasLegacyData ? ActorSnapshots.Num() : SerializedData.GetNumSavedActors();
+	return SerializedData.GetNumSavedActors();
 }
 
 void ULevelSnapshot::DiffWorld(UWorld* World, FActorPathConsumer HandleMatchedActor, FActorPathConsumer HandleRemovedActor, FActorConsumer HandleAddedActor) const
@@ -326,16 +298,6 @@ void ULevelSnapshot::DiffWorld(UWorld* World, FActorPathConsumer HandleMatchedAc
 		return;
 	}
 
-	const bool bHasLegacyData = ActorSnapshots.Num() > 0;
-	if (bHasLegacyData)
-	{
-		for (auto It = ActorSnapshots.CreateConstIterator(); It; ++It)
-		{
-			HandleMatchedActor.Execute(It->Key);
-		}
-		return;
-	}
-
 	// Find actors that are not present in the snapshot
 	TSet<AActor*> AllActors;
 	const int32 NumActorsInWorld = Algo::Accumulate(World->GetLevels(), 0, [](int64 Size, const ULevel* Level){ return Size + Level->Actors.Num(); });
@@ -465,49 +427,4 @@ void ULevelSnapshot::DestroyWorld()
 		SnapshotContainerWorld->CleanupWorld();
 		SnapshotContainerWorld = nullptr;
 	}
-}
-
-void ULevelSnapshot::LegacyApplySnapshotToWorld(const FPropertySelectionMap& SelectionSet)
-{
-	DECLARE_SCOPE_CYCLE_COUNTER(TEXT("ApplySnapshotToWorld"), STAT_ApplySnapshotToWorld, STATGROUP_LevelSnapshots);
-
-	ULevelSnapshotSelectionSet* LegacySelectionSet = NewObject<ULevelSnapshotSelectionSet>();
-	LegacySelectionSet->AddPropertyMap(SelectionSet);
-	
-	TSet<AActor*> EvaluatedActors;
-	for (const FSoftObjectPath& Path : LegacySelectionSet->GetSelectedWorldObjectPaths())
-	{
-		if (Path.IsValid())
-		{
-			AActor* ActorToPass = nullptr;
-				
-			if (UObject* SelectedObject = Path.ResolveObject())
-			{
-				if (AActor* AsActor = Cast<AActor>(SelectedObject))
-				{
-					ActorToPass = AsActor;
-				}
-				else if (UActorComponent* AsComponent = Cast<UActorComponent>(SelectedObject))
-				{
-					if (AActor* OwningActor = AsComponent->GetOwner())
-					{
-						ActorToPass = OwningActor;
-					}
-				}
-			}
-
-			if (ensure(ActorToPass))
-			{
-				if (!EvaluatedActors.Contains(ActorToPass))
-				{
-					if (const FLevelSnapshot_Actor* ActorSnapshot = ActorSnapshots.Find(ActorToPass))
-					{
-						ActorSnapshot->DeserializeIntoWorldActor(ActorToPass, LegacySelectionSet);
-					}
-						
-					EvaluatedActors.Add(ActorToPass);
-				}
-			}
-		}
-	}
 }
\ No newline at end of file
diff --git a/Engine/Plugins/VirtualProduction/LevelSnapshots/Source/LevelSnapshots/Private/LevelSnapshotSelections.cpp b/Engine/Plugins/VirtualProduction/LevelSnapshots/Source/LevelSnapshots/Private/LevelSnapshotSelections.cpp
deleted file mode 100644
diff --git a/Engine/Plugins/VirtualProduction/LevelSnapshots/Source/LevelSnapshots/Public/Data/ActorSnapshot.h b/Engine/Plugins/VirtualProduction/LevelSnapshots/Source/LevelSnapshots/Public/Data/ActorSnapshot.h
deleted file mode 100644
diff --git a/Engine/Plugins/VirtualProduction/LevelSnapshots/Source/LevelSnapshots/Public/Data/BaseObjectInfo.h b/Engine/Plugins/VirtualProduction/LevelSnapshots/Source/LevelSnapshots/Public/Data/BaseObjectInfo.h
deleted file mode 100644
diff --git a/Engine/Plugins/VirtualProduction/LevelSnapshots/Source/LevelSnapshots/Public/Data/LevelSnapshot.h b/Engine/Plugins/VirtualProduction/LevelSnapshots/Source/LevelSnapshots/Public/Data/LevelSnapshot.h
--- a/Engine/Plugins/VirtualProduction/LevelSnapshots/Source/LevelSnapshots/Public/Data/LevelSnapshot.h
+++ b/Engine/Plugins/VirtualProduction/LevelSnapshots/Source/LevelSnapshots/Public/Data/LevelSnapshot.h
@@ -3,7 +3,6 @@
 #pragma once
 
 #include "CoreMinimal.h"
-#include "ActorSnapshot.h"
 #include "WorldSnapshotData.h"
 #include "LevelSnapshot.generated.h"
 
@@ -75,16 +74,6 @@ private:
 
 	void EnsureWorldInitialised();
 	void DestroyWorld();
-	
-	/****************************** Start legacy members ******************************/
-
-	void LegacyApplySnapshotToWorld(const FPropertySelectionMap& SelectionSet);
-	
-	// Map of Actor Snapshots mapping from the object path to the actual snapshot
-	UPROPERTY(VisibleAnywhere, Category = "Snapshot|Deprecated")
-	TMap<FSoftObjectPath, FLevelSnapshot_Actor> ActorSnapshots;
-
-	/****************************** End legacy members ******************************/
 
 	
 	/* The world we will be adding temporary actors to */
diff --git a/Engine/Plugins/VirtualProduction/LevelSnapshots/Source/LevelSnapshots/Public/LevelSnapshotSelections.h b/Engine/Plugins/VirtualProduction/LevelSnapshots/Source/LevelSnapshots/Public/LevelSnapshotSelections.h
deleted file mode 100644
diff --git a/Engine/Plugins/VirtualProduction/LevelSnapshots/Source/LevelSnapshots/Public/LevelSnapshotsFunctionLibrary.h b/Engine/Plugins/VirtualProduction/LevelSnapshots/Source/LevelSnapshots/Public/LevelSnapshotsFunctionLibrary.h
--- a/Engine/Plugins/VirtualProduction/LevelSnapshots/Source/LevelSnapshots/Public/LevelSnapshotsFunctionLibrary.h
+++ b/Engine/Plugins/VirtualProduction/LevelSnapshots/Source/LevelSnapshots/Public/LevelSnapshotsFunctionLibrary.h
@@ -4,7 +4,6 @@
 
 #include "CoreMinimal.h"
 #include "Kismet/BlueprintFunctionLibrary.h"
-#include "LevelSnapshotSelections.h"
 #include "PropertySelectionMap.h"
 #include "LevelSnapshotsFunctionLibrary.generated.h"
 
