Crash in atmospheric fog component

diff --git a/Engine/Plugins/VirtualProduction/LevelSnapshots/Source/LevelSnapshots/Private/Data/ActorSnapshotData.cpp b/Engine/Plugins/VirtualProduction/LevelSnapshots/Source/LevelSnapshots/Private/Data/ActorSnapshotData.cpp
--- a/Engine/Plugins/VirtualProduction/LevelSnapshots/Source/LevelSnapshots/Private/Data/ActorSnapshotData.cpp
+++ b/Engine/Plugins/VirtualProduction/LevelSnapshots/Source/LevelSnapshots/Private/Data/ActorSnapshotData.cpp
@@ -107,11 +107,8 @@ TOptional<AActor*> FActorSnapshotData::GetPreallocated(UWorld* SnapshotWorld, FW
 		FActorSpawnParameters SpawnParams;
 		SpawnParams.Template = Cast<AActor>(WorldData.GetClassDefault(TargetClass));
 		SpawnParams.NameMode = FActorSpawnParameters::ESpawnActorNameMode::Requested;
-		if (ensureMsgf(SpawnParams.Template, TEXT("Failed to class default. This should not happen. Investigate.")))
+		if (ensureMsgf(SpawnParams.Template, TEXT("Failed to get class default. This should not happen. Investigate.")))
 		{
-			// We're passing in SpawnParams.Template->GetClass() instead of TargetClass:
-				// When you recompile a Blueprint, it creates a new REINST class.
-				// This would cause the SpawnParams.Template to have a different class than TargetClass: that would cause SpawnActor to fail.
 			UClass* ClassToUse = SpawnParams.Template->GetClass();
 			SpawnParams.Name = *FString("SnapshotObjectInstance_").Append(*MakeUniqueObjectName(SnapshotWorld, ClassToUse).ToString());
 			CachedSnapshotActor = SnapshotWorld->SpawnActor<AActor>(ClassToUse, SpawnParams);
@@ -122,7 +119,14 @@ TOptional<AActor*> FActorSnapshotData::GetPreallocated(UWorld* SnapshotWorld, FW
 		}
 	}
 
-	ensureAlwaysMsgf(CachedSnapshotActor.IsValid(), TEXT("Failed to spawn actor of class '%s'"), *ActorClass.ToString());
+#if WITH_EDITOR
+	// Hide this actor so external systems can see that this components should not render, i.e. make USceneComponent::ShouldRender return false
+	if (ensureMsgf(CachedSnapshotActor.IsValid(), TEXT("Failed to spawn actor of class '%s'"), *ActorClass.ToString()))
+	{
+		CachedSnapshotActor->SetIsTemporarilyHiddenInEditor(true);
+	}
+#endif
+	
 	return CachedSnapshotActor.IsValid() ? TOptional<AActor*>(CachedSnapshotActor.Get()) : TOptional<AActor*>();
 }
 
@@ -149,6 +153,14 @@ TOptional<AActor*> FActorSnapshotData::GetDeserialized(UWorld* SnapshotWorld, FW
 	});
 
 	PreallocatedActor->UpdateComponentTransforms();
+#if WITH_EDITOR
+	// Hide this actor so external systems can see that this components should not render, i.e. make USceneComponent::ShouldRender return false
+	if (!ensureMsgf(PreallocatedActor->IsTemporarilyHiddenInEditor(), TEXT("Transient property bHiddenEdTemporary was set to false by serializer. This should not happen. Investigate.")))
+	{
+		CachedSnapshotActor->SetIsTemporarilyHiddenInEditor(true);
+	}
+#endif
+	
 	return Preallocated;
 }
 
