Fix static nullptr check error

diff --git a/Engine/Plugins/VirtualProduction/LevelSnapshots/Source/LevelSnapshots/Private/Data/WorldSnapshotData.cpp b/Engine/Plugins/VirtualProduction/LevelSnapshots/Source/LevelSnapshots/Private/Data/WorldSnapshotData.cpp
--- a/Engine/Plugins/VirtualProduction/LevelSnapshots/Source/LevelSnapshots/Private/Data/WorldSnapshotData.cpp
+++ b/Engine/Plugins/VirtualProduction/LevelSnapshots/Source/LevelSnapshots/Private/Data/WorldSnapshotData.cpp
@@ -293,7 +293,10 @@ bool FWorldSnapshotData::AreReferencesEquivalent(UObject* SnapshotPropertyValue,
 	const bool bNeedsUnsupportedSubobjectRestorationFeature = SnapshotPropertyValue == nullptr && bOriginalIsSubobject;
 	if (bNeedsUnsupportedSubobjectRestorationFeature)
 	{
-		UE_LOG(LogLevelSnapshots, Verbose, TEXT("Object '%s' in of actor '%s' seems to be a subobject. Snapshots currently do not support re-creating subobjects."), *OriginalPropertyValue->GetName(), *OriginalActor->GetName());
+		UE_LOG(LogLevelSnapshots, Verbose, TEXT("Object '%s' in of actor '%s' seems to be a subobject. Snapshots currently do not support re-creating subobjects."),
+			OriginalPropertyValue ? *OriginalPropertyValue->GetName() : TEXT("null"),
+			*OriginalActor->GetName()
+			);
 		return true;
 	}
 	
