Show filter description as tooltip

diff --git a/Engine/Plugins/VirtualProduction/LevelSnapshots/Source/LevelSnapshotsEditor/Private/Widgets/Filter/SFilterSearchMenu.cpp b/Engine/Plugins/VirtualProduction/LevelSnapshots/Source/LevelSnapshotsEditor/Private/Widgets/Filter/SFilterSearchMenu.cpp
--- a/Engine/Plugins/VirtualProduction/LevelSnapshots/Source/LevelSnapshotsEditor/Private/Widgets/Filter/SFilterSearchMenu.cpp
+++ b/Engine/Plugins/VirtualProduction/LevelSnapshots/Source/LevelSnapshotsEditor/Private/Widgets/Filter/SFilterSearchMenu.cpp
@@ -53,7 +53,7 @@ namespace
 		
 				MenuBuilder.AddMenuEntry(
                     FilterClass->GetDisplayNameText(),
-                    FilterClass->GetDisplayNameText(),
+                    FilterClass->GetToolTipText(),
                     FSlateIcon(),
                     AddFilterAction,
                     NAME_None,
