Misc:
-Added a few functions to facilitate clearing the selected world and selected snapshot from places they're stored when the LS dialogue is initialized. This fixes an emergent issue in which the restore and refresh buttons would load the last selected snapshot on a new instance of the toolkit as this data persists between instances.
-Fixed an oversight in the Toolkit script in which two booleans were in the wrong spot. They were both set to false so this issue did not have any felt ramifications, but it needed correction nonetheless.

Results View:
-Now all rows hold a weak reference to the Results View so that they can get the Show Options to know whether or not they should be visible to the user.
-Rows now will remember their expansion and selection states when results are refreshed or partially restored. They will not keep these states if the snapshot changes by design. If the number of expanded actor groups stays relatively low, this should not be an issue but if many are open there will be a hitch proportional to the number of open groups that remain after refresh or restore.
-Added a bool to SetWidgetCheckedState that determines whether or not the new setting should affect the parents or children related to the row in question. This is useful when restoring checked states from memory - setting child and parent checked states repeatedly has an unnecessary performance cost.
-Implemented the Show Option for Show Unselected. Now if that option is set to false, any rows whose checkboxes are unchecked will be invisible to the user.
-Fixed an issue in which results were not cleared when the editor world was changed.
-Disabled the Restore button if EditorData didn't have an active snapshot.
-Added option to keep memory allocation when flushing memory in some cases.
-Changed Show Options language from "Actors" to "Rows" to better reflect the functionality.
-Added function to update title text of selected snapshot.
-Added pointer check when building a selection set on restore
-Added function to expand a tree view item to Results View. This is necessary for maintaining expanded states between refresh as expansion is facilitated through the tree view, not the row data itself.

#[fyi] Jason.Walter

diff --git a/Engine/Plugins/VirtualProduction/LevelSnapshots/Source/LevelSnapshotsEditor/Private/Data/FilteredResults.cpp b/Engine/Plugins/VirtualProduction/LevelSnapshots/Source/LevelSnapshotsEditor/Private/Data/FilteredResults.cpp
--- a/Engine/Plugins/VirtualProduction/LevelSnapshots/Source/LevelSnapshotsEditor/Private/Data/FilteredResults.cpp
+++ b/Engine/Plugins/VirtualProduction/LevelSnapshots/Source/LevelSnapshotsEditor/Private/Data/FilteredResults.cpp
@@ -60,6 +60,12 @@ TWeakObjectPtr<ULevelSnapshotFilter> UFilteredResults::GetUserFilters() const
 	return UserFilters;
 }
 
+void UFilteredResults::ClearSelectedWorld()
+{
+	SelectedWorld = nullptr;
+	CleanReferences();
+}
+
 void UFilteredResults::SetSelectedWorld(UWorld* InWorld)
 {
 	SelectedWorld = InWorld;
diff --git a/Engine/Plugins/VirtualProduction/LevelSnapshots/Source/LevelSnapshotsEditor/Private/Data/FilteredResults.h b/Engine/Plugins/VirtualProduction/LevelSnapshots/Source/LevelSnapshotsEditor/Private/Data/FilteredResults.h
--- a/Engine/Plugins/VirtualProduction/LevelSnapshots/Source/LevelSnapshotsEditor/Private/Data/FilteredResults.h
+++ b/Engine/Plugins/VirtualProduction/LevelSnapshots/Source/LevelSnapshotsEditor/Private/Data/FilteredResults.h
@@ -34,6 +34,7 @@ public:
 	FFilterListData& GetFilteredData();
 	TWeakObjectPtr<ULevelSnapshotFilter> GetUserFilters() const;
 
+	void ClearSelectedWorld();
 	void SetSelectedWorld(UWorld* InWorld);
 
 private:
diff --git a/Engine/Plugins/VirtualProduction/LevelSnapshots/Source/LevelSnapshotsEditor/Private/Data/LevelSnapshotsEditorData.cpp b/Engine/Plugins/VirtualProduction/LevelSnapshots/Source/LevelSnapshotsEditor/Private/Data/LevelSnapshotsEditorData.cpp
--- a/Engine/Plugins/VirtualProduction/LevelSnapshots/Source/LevelSnapshotsEditor/Private/Data/LevelSnapshotsEditorData.cpp
+++ b/Engine/Plugins/VirtualProduction/LevelSnapshots/Source/LevelSnapshotsEditor/Private/Data/LevelSnapshotsEditorData.cpp
@@ -75,6 +75,13 @@ void ULevelSnapshotsEditorData::SetActiveSnapshot(const TOptional<ULevelSnapshot
 	OnActiveSnapshotChanged.Broadcast(GetActiveSnapshot());
 }
 
+void ULevelSnapshotsEditorData::ClearActiveSnapshot()
+{
+	ActiveSnapshot.Reset();
+	FilterResults->SetActiveLevelSnapshot(nullptr);
+	OnActiveSnapshotChanged.Broadcast(TOptional<ULevelSnapshot*>(nullptr));
+}
+
 TOptional<ULevelSnapshot*> ULevelSnapshotsEditorData::GetActiveSnapshot() const
 {
 	return ActiveSnapshot.IsSet() ? TOptional<ULevelSnapshot*>(ActiveSnapshot->Get()) : TOptional<ULevelSnapshot*>();
@@ -87,6 +94,13 @@ void ULevelSnapshotsEditorData::SetSelectedWorldReference(UWorld* InWorld)
 	FilterResults->SetSelectedWorld(InWorld);
 }
 
+void ULevelSnapshotsEditorData::ClearSelectedWorld()
+{
+	SelectedWorld = nullptr;
+	
+	FilterResults->ClearSelectedWorld();
+}
+
 UWorld* ULevelSnapshotsEditorData::GetSelectedWorld() const
 {
 	return SelectedWorld.IsSet() ? SelectedWorld.GetValue() : nullptr;
diff --git a/Engine/Plugins/VirtualProduction/LevelSnapshots/Source/LevelSnapshotsEditor/Private/Data/LevelSnapshotsEditorData.h b/Engine/Plugins/VirtualProduction/LevelSnapshots/Source/LevelSnapshotsEditor/Private/Data/LevelSnapshotsEditorData.h
--- a/Engine/Plugins/VirtualProduction/LevelSnapshots/Source/LevelSnapshotsEditor/Private/Data/LevelSnapshotsEditorData.h
+++ b/Engine/Plugins/VirtualProduction/LevelSnapshots/Source/LevelSnapshotsEditor/Private/Data/LevelSnapshotsEditorData.h
@@ -36,6 +36,7 @@ public:
 	
 	/******************** Active snapshot ********************/
 	void SetActiveSnapshot(const TOptional<ULevelSnapshot*>& NewActiveSnapshot);
+	void ClearActiveSnapshot();
 	TOptional<ULevelSnapshot*> GetActiveSnapshot() const;
 
 	DECLARE_EVENT_OneParam(ULevelSnapshotsEditorData, FOnActiveSnapshotChanged, const TOptional<ULevelSnapshot*>& /* NewSnapshot */);
@@ -45,6 +46,7 @@ public:
 	
 	/******************** Selected world ********************/
 	void SetSelectedWorldReference(UWorld* InWorld);
+	void ClearSelectedWorld();
 	UWorld* GetSelectedWorld() const;
 
 
diff --git a/Engine/Plugins/VirtualProduction/LevelSnapshots/Source/LevelSnapshotsEditor/Private/Toolkits/LevelSnapshotsEditorToolkit.cpp b/Engine/Plugins/VirtualProduction/LevelSnapshots/Source/LevelSnapshotsEditor/Private/Toolkits/LevelSnapshotsEditorToolkit.cpp
--- a/Engine/Plugins/VirtualProduction/LevelSnapshots/Source/LevelSnapshotsEditor/Private/Toolkits/LevelSnapshotsEditorToolkit.cpp
+++ b/Engine/Plugins/VirtualProduction/LevelSnapshots/Source/LevelSnapshotsEditor/Private/Toolkits/LevelSnapshotsEditorToolkit.cpp
@@ -40,8 +40,12 @@ TSharedPtr<FLevelSnapshotsEditorToolkit> FLevelSnapshotsEditorToolkit::CreateSna
 
 void FLevelSnapshotsEditorToolkit::Initialize(ULevelSnapshotsEditorData* InEditorData)
 {
+	check(InEditorData);
+	
 	EditorData = InEditorData;
 
+	EditorData.Get()->ClearActiveSnapshot();
+
 	EditorContext = MakeShared<FLevelSnapshotsEditorContext>();
 	ViewBuilder = MakeShared<FLevelSnapshotsEditorViewBuilder>();
 	ViewBuilder->EditorContextPtr = EditorContext;
@@ -108,8 +112,8 @@ void FLevelSnapshotsEditorToolkit::Initialize(ULevelSnapshotsEditorData* InEdito
 			nullptr,
 			AppIdentifier,
 			Layout,
-			bCreateDefaultToolbar,
 			bCreateDefaultStandaloneMenu,
+			bCreateDefaultToolbar,
 			InEditorData
 			);
 }
diff --git a/Engine/Plugins/VirtualProduction/LevelSnapshots/Source/LevelSnapshotsEditor/Private/Views/Input/SlevelSnapshotsEditorContextPicker.cpp b/Engine/Plugins/VirtualProduction/LevelSnapshots/Source/LevelSnapshotsEditor/Private/Views/Input/SlevelSnapshotsEditorContextPicker.cpp
--- a/Engine/Plugins/VirtualProduction/LevelSnapshots/Source/LevelSnapshotsEditor/Private/Views/Input/SlevelSnapshotsEditorContextPicker.cpp
+++ b/Engine/Plugins/VirtualProduction/LevelSnapshots/Source/LevelSnapshotsEditor/Private/Views/Input/SlevelSnapshotsEditorContextPicker.cpp
@@ -241,6 +241,11 @@ void SLevelSnapshotsEditorContextPicker::OnSetWorldContextSelection(const FAsset
 
 void SLevelSnapshotsEditorContextPicker::SetSelectedWorld(const FSoftObjectPath& SelectedWorld)
 {
+	if (SelectedWorldPath == SelectedWorld)
+	{
+		return;
+	}
+	
 	SelectedWorldPath = SelectedWorld;
 	OnSelectWorldContextEvent.ExecuteIfBound(SelectedWorld);
 
@@ -250,6 +255,7 @@ void SLevelSnapshotsEditorContextPicker::SetSelectedWorld(const FSoftObjectPath&
 		{
 			const TWeakObjectPtr<ULevelSnapshotsEditorData> EditorData = Builder->EditorDataPtr;
 			EditorData->SetSelectedWorldReference(GetSelectedWorld());
+			EditorData->SetActiveSnapshot(nullptr);
 		}
 	}
 	else
diff --git a/Engine/Plugins/VirtualProduction/LevelSnapshots/Source/LevelSnapshotsEditor/Private/Views/Results/SLevelSnapshotsEditorResults.cpp b/Engine/Plugins/VirtualProduction/LevelSnapshots/Source/LevelSnapshotsEditor/Private/Views/Results/SLevelSnapshotsEditorResults.cpp
--- a/Engine/Plugins/VirtualProduction/LevelSnapshots/Source/LevelSnapshotsEditor/Private/Views/Results/SLevelSnapshotsEditorResults.cpp
+++ b/Engine/Plugins/VirtualProduction/LevelSnapshots/Source/LevelSnapshotsEditor/Private/Views/Results/SLevelSnapshotsEditorResults.cpp
@@ -12,6 +12,7 @@
 
 #include "Algo/Find.h"
 #include "DebugViewModeHelpers.h"
+#include "Editor.h"
 #include "EditorStyleSet.h"
 #include "Framework/MultiBox/MultiBoxBuilder.h"
 #include "Framework/Notifications/NotificationManager.h"
@@ -134,61 +135,67 @@ void FLevelSnapshotsEditorResultsRow::FlushReferences()
 FLevelSnapshotsEditorResultsRow::FLevelSnapshotsEditorResultsRow(
 	const FText InDisplayName,                                     
 	const ELevelSnapshotsEditorResultsRowType InRowType, const ECheckBoxState StartingWidgetCheckboxState, 
-	const TWeakPtr<FLevelSnapshotsEditorResultsRow>& DirectParent)
+	const TWeakPtr<SLevelSnapshotsEditorResults>& InResultsView, const TWeakPtr<FLevelSnapshotsEditorResultsRow>& InDirectParentRow)
 {
+	check(InResultsView.IsValid());
+	ResultsViewPtr = InResultsView;
+	
 	DisplayName = InDisplayName;
 	RowType = InRowType;
-	WidgetCheckedState = StartingWidgetCheckboxState;
+	SetWidgetCheckedState(StartingWidgetCheckboxState);
 
-	if (DirectParent.IsValid())
+	if (InDirectParentRow.IsValid())
 	{
-		DirectParentRow = DirectParent;
+		DirectParentRow = InDirectParentRow;
 	}
 }
 
 void FLevelSnapshotsEditorResultsRow::InitHeaderRow(const ELevelSnapshotsEditorResultsTreeViewHeaderType InHeaderType,
-	const TArray<FText>& InColumns, const TWeakPtr<SLevelSnapshotsEditorResults>& InResultsView)
+	const TArray<FText>& InColumns)
 {
 	HeaderType = InHeaderType;
 	HeaderColumns = InColumns;
-	ResultsViewPtr = InResultsView;
 
 	// Set visibility values to false for header rows so they hide correctly when they have no visible children
 	bDoesRowMatchSearchTerms = false;
+
+	ApplyRowStateMemoryIfAvailable();
 }
 
 void FLevelSnapshotsEditorResultsRow::InitAddedActorRow(AActor* InAddedActor)
 {
 	WorldObject = InAddedActor;
+
+	ApplyRowStateMemoryIfAvailable();
 }
 
-void FLevelSnapshotsEditorResultsRow::InitRemovedActorRow(const FSoftObjectPath& InRemovedActorPath,
-                                                          const TWeakPtr<SLevelSnapshotsEditorResults>& InResultsView)
+void FLevelSnapshotsEditorResultsRow::InitRemovedActorRow(const FSoftObjectPath& InRemovedActorPath)
 {
 	RemovedActorPath = InRemovedActorPath;
-	ResultsViewPtr = InResultsView;
+
+	ApplyRowStateMemoryIfAvailable();
 }
 
 void FLevelSnapshotsEditorResultsRow::InitActorRow(
-	AActor* InSnapshotActor, AActor* InWorldActor,
-	const TWeakPtr<SLevelSnapshotsEditorResults>& InResultsView)
+	AActor* InSnapshotActor, AActor* InWorldActor)
 {
 	SnapshotObject = InSnapshotActor;
 	WorldObject = InWorldActor;
-	ResultsViewPtr = InResultsView;
+
+	ApplyRowStateMemoryIfAvailable();
 }
 
 void FLevelSnapshotsEditorResultsRow::InitObjectRow(
 	UObject* InSnapshotObject, UObject* InWorldObject,
 	const TWeakPtr<FRowGeneratorInfo>& InSnapshotRowGenerator,
-	const TWeakPtr<FRowGeneratorInfo>& InWorldRowGenerator,
-	const TWeakPtr<SLevelSnapshotsEditorResults>& InResultsView)
+	const TWeakPtr<FRowGeneratorInfo>& InWorldRowGenerator)
 {
 	SnapshotObject = InSnapshotObject;
 	WorldObject = InWorldObject;
 	SnapshotRowGeneratorInfo = InSnapshotRowGenerator;
 	WorldRowGeneratorInfo = InWorldRowGenerator;
-	ResultsViewPtr = InResultsView;
+
+	ApplyRowStateMemoryIfAvailable();
 }
 
 void FLevelSnapshotsEditorResultsRow::InitPropertyRow(
@@ -200,6 +207,42 @@ void FLevelSnapshotsEditorResultsRow::InitPropertyRow(
 	SnapshotPropertyHandleHierarchy = InSnapshotHierarchy;
 	WorldPropertyHandleHierarchy = InWorldHandleHierarchy;
 	bIsCounterpartValueSame = bNewIsCounterpartValueSame;
+
+	ApplyRowStateMemoryIfAvailable();
+}
+
+void FLevelSnapshotsEditorResultsRow::ApplyRowStateMemoryIfAvailable()
+{
+	FLevelSnapshotsEditorResultsRowStateMemory RowStateMemory;
+	const TSharedPtr<SLevelSnapshotsEditorResults>& ResultsPinned = ResultsViewPtr.Pin();
+	
+	if (ResultsPinned->FindRowStateMemoryByPath(GetOrGenerateRowPath(), RowStateMemory))
+	{
+		if (DoesRowRepresentGroup())
+		{
+			if (RowStateMemory.bIsExpanded)
+			{
+				ResultsPinned->SetTreeViewItemExpanded(SharedThis(this), true);
+			}
+		}
+
+		SetWidgetCheckedState(RowStateMemory.WidgetCheckedState);
+	}
+}
+
+const FString& FLevelSnapshotsEditorResultsRow::GetOrGenerateRowPath()
+{
+	if (RowPath.IsEmpty())
+	{
+		RowPath = GetDisplayName().ToString();
+
+		if (DirectParentRow.IsValid())
+		{
+			RowPath = DirectParentRow.Pin()->GetOrGenerateRowPath() + "." + RowPath;
+		}
+	}
+
+	return RowPath;
 }
 
 void FLevelSnapshotsEditorResultsRow::GenerateActorGroupChildren(FPropertySelectionMap& PropertySelectionMap)
@@ -208,7 +251,8 @@ void FLevelSnapshotsEditorResultsRow::GenerateActorGroupChildren(FPropertySelect
 	{
 		static TArray<TFieldPath<FProperty>> LoopOverProperties(
 			const TWeakPtr<FRowGeneratorInfo>& InSnapshotRowGeneratorInfo, const TWeakPtr<FRowGeneratorInfo>& InWorldRowGeneratorInfo,
-			const TWeakPtr<FLevelSnapshotsEditorResultsRow>& InDirectParentRow, const TArray<TFieldPath<FProperty>>& PropertiesThatPassFilter)
+			const TWeakPtr<FLevelSnapshotsEditorResultsRow>& InDirectParentRow, const TArray<TFieldPath<FProperty>>& PropertiesThatPassFilter,
+			const TWeakPtr<SLevelSnapshotsEditorResults>& InResultsView)
 		{
 			TSharedPtr<FPropertyHandleHierarchy> SnapshotHandleHierarchy = nullptr;
 			TSharedPtr<FPropertyHandleHierarchy> WorldHandleHierarchy = nullptr;
@@ -234,7 +278,7 @@ void FLevelSnapshotsEditorResultsRow::GenerateActorGroupChildren(FPropertySelect
 					const TSharedRef<FPropertyHandleHierarchy>& ChildHierarchy = WorldHandleHierarchy->DirectChildren[ChildIndex];
 
 					LoopOverHandleHierarchiesAndCreateRowHierarchy(
-							ChildHierarchy, InDirectParentRow, PropertiesThatPassFilter, PropertyRowsGenerated, SnapshotHandleHierarchy);
+							ChildHierarchy, InDirectParentRow, PropertiesThatPassFilter, PropertyRowsGenerated, InResultsView, SnapshotHandleHierarchy);
 				}
 			}
 
@@ -244,7 +288,7 @@ void FLevelSnapshotsEditorResultsRow::GenerateActorGroupChildren(FPropertySelect
 		/* Do not pass in the base hierarchy, only children of the base hierarchy. */
 		static void LoopOverHandleHierarchiesAndCreateRowHierarchy(const TWeakPtr<FPropertyHandleHierarchy>& InHierarchy,
 			const TWeakPtr<FLevelSnapshotsEditorResultsRow>& InDirectParentRow, const TArray<TFieldPath<FProperty>>& PropertiesThatPassFilter,
-			TArray<TFieldPath<FProperty>>& PropertyRowsGenerated, const TWeakPtr<FPropertyHandleHierarchy>& InHierarchyToSearchForCounterparts = nullptr)
+			TArray<TFieldPath<FProperty>>& PropertyRowsGenerated, const TWeakPtr<SLevelSnapshotsEditorResults>& InResultsView, const TWeakPtr<FPropertyHandleHierarchy>& InHierarchyToSearchForCounterparts = nullptr)
 		{
 			if (!ensure(InHierarchy.IsValid()) || !ensure(InDirectParentRow.IsValid()))
 			{
@@ -315,7 +359,7 @@ void FLevelSnapshotsEditorResultsRow::GenerateActorGroupChildren(FPropertySelect
 
 				// Create property
 				FLevelSnapshotsEditorResultsRowPtr NewProperty = MakeShared<FLevelSnapshotsEditorResultsRow>(DisplayName, InRowType, StartingCheckedState, 
-					InDirectParentRow);
+					InResultsView, InDirectParentRow);
 
 				const TWeakPtr<FLevelSnapshotsEditorResultsRow>& ContainingObjectGroup = InDirectParentRow;
 
@@ -327,7 +371,7 @@ void FLevelSnapshotsEditorResultsRow::GenerateActorGroupChildren(FPropertySelect
 				{
 					const TSharedRef<FPropertyHandleHierarchy>& ChildHierarchy = InHierarchy.Pin()->DirectChildren[ChildIndex];
 
-					LoopOverHandleHierarchiesAndCreateRowHierarchy(ChildHierarchy, NewProperty, PropertiesThatPassFilter, PropertyRowsGenerated, InHierarchyToSearchForCounterparts);
+					LoopOverHandleHierarchiesAndCreateRowHierarchy(ChildHierarchy, NewProperty, PropertiesThatPassFilter, PropertyRowsGenerated, InResultsView, InHierarchyToSearchForCounterparts);
 				}
 
 				if (NewProperty->DoesRowRepresentGroup() && !NewProperty->GetChildRows().Num())
@@ -630,7 +674,7 @@ void FLevelSnapshotsEditorResultsRow::GenerateActorGroupChildren(FPropertySelect
 			// Create group
 			FLevelSnapshotsEditorResultsRowPtr NewComponentGroup = MakeShared<FLevelSnapshotsEditorResultsRow>(
 				FText::FromString(InComponent->GetName()), ComponentGroup,
-				InDirectParentRow.IsValid() ? InDirectParentRow.Pin()->GetWidgetCheckedState() : ECheckBoxState::Checked, InDirectParentRow);
+				InDirectParentRow.IsValid() ? InDirectParentRow.Pin()->GetWidgetCheckedState() : ECheckBoxState::Checked, InResultsView, InDirectParentRow);
 
 			// Create Row Generators for object and counterpart
 			const TWeakPtr<FRowGeneratorInfo>& RowGeneratorInfo = InResultsView.Pin()->RegisterRowGenerator(NewComponentGroup, ObjectType_World, PropertyEditorModule);
@@ -646,9 +690,9 @@ void FLevelSnapshotsEditorResultsRow::GenerateActorGroupChildren(FPropertySelect
 				CounterpartRowGeneratorInfo.Pin()->GetGeneratorObject().Pin()->SetObjects({ CounterpartComponent });
 			}
 			
-			NewComponentGroup->InitObjectRow(InComponent, CounterpartComponent, CounterpartRowGeneratorInfo, RowGeneratorInfo, InResultsView);
+			NewComponentGroup->InitObjectRow(InComponent, CounterpartComponent, CounterpartRowGeneratorInfo, RowGeneratorInfo);
 
-			const TArray<TFieldPath<FProperty>> PropertyRowsGenerated = FLocalPropertyLooper::LoopOverProperties(CounterpartRowGeneratorInfo, RowGeneratorInfo, NewComponentGroup, PropertiesThatPassFilter);
+			const TArray<TFieldPath<FProperty>> PropertyRowsGenerated = FLocalPropertyLooper::LoopOverProperties(CounterpartRowGeneratorInfo, RowGeneratorInfo, NewComponentGroup, PropertiesThatPassFilter, InResultsView);
 
 			// Generate fallback rows for properties not supported by PropertyRowGenerator
 			for (TFieldPath<FProperty> FieldPath : PropertiesThatPassFilter)
@@ -730,7 +774,7 @@ void FLevelSnapshotsEditorResultsRow::GenerateActorGroupChildren(FPropertySelect
 			const TArray<TFieldPath<FProperty>>& PropertyRowsGenerated = FLocalPropertyLooper::LoopOverProperties(
 				SnapshotRowGeneratorInfo,
 				WorldRowGeneratorInfo, 
-				SharedThis(this), PropertySelection->GetSelectedLeafProperties());
+				SharedThis(this), PropertySelection->GetSelectedLeafProperties(), ResultsViewPtr);
 			
 			// Generate fallback rows for properties not supported by 
 			for (TFieldPath<FProperty> FieldPath : PropertySelection->GetSelectedLeafProperties())
@@ -1183,23 +1227,26 @@ ECheckBoxState FLevelSnapshotsEditorResultsRow::GetWidgetCheckedState() const
 	return WidgetCheckedState;
 }
 
-void FLevelSnapshotsEditorResultsRow::SetWidgetCheckedState(const ECheckBoxState NewState, const bool bUserClicked)
+void FLevelSnapshotsEditorResultsRow::SetWidgetCheckedState(const ECheckBoxState NewState, const bool bShouldUpdateHierarchyCheckedStates)
 {
 	WidgetCheckedState = NewState;
 
-	// Set Children to same checked state
-	for (const FLevelSnapshotsEditorResultsRowPtr& ChildRow : GetChildRows())
+	if (bShouldUpdateHierarchyCheckedStates)
 	{
-		if (!ChildRow.IsValid())
+		// Set Children to same checked state
+		for (const FLevelSnapshotsEditorResultsRowPtr& ChildRow : GetChildRows())
 		{
-			continue;
-		}
+			if (!ChildRow.IsValid())
+			{
+				continue;
+			}
 
-		ChildRow->SetWidgetCheckedState(NewState);
+			ChildRow->SetWidgetCheckedState(NewState);
+		}
+		
+		EvaluateAndSetAllParentGroupCheckedStates();
 	}
 
-	EvaluateAndSetAllParentGroupCheckedStates();
-
 	if (GetRowType() == ActorGroup && ResultsViewPtr.IsValid())
 	{
 		ResultsViewPtr.Pin()->UpdateSnapshotInformationText();
@@ -1211,9 +1258,9 @@ bool FLevelSnapshotsEditorResultsRow::GetIsNodeChecked() const
 	return GetWidgetCheckedState() == ECheckBoxState::Checked ? true : false;
 }
 
-void FLevelSnapshotsEditorResultsRow::SetIsNodeChecked(const bool bNewChecked)
+void FLevelSnapshotsEditorResultsRow::SetIsNodeChecked(const bool bNewChecked, const bool bShouldUpdateHierarchyCheckedStates)
 {
-	SetWidgetCheckedState(bNewChecked ? ECheckBoxState::Checked : ECheckBoxState::Unchecked);
+	SetWidgetCheckedState(bNewChecked ? ECheckBoxState::Checked : ECheckBoxState::Unchecked, bShouldUpdateHierarchyCheckedStates);
 }
 
 bool FLevelSnapshotsEditorResultsRow::HasVisibleChildren() const
@@ -1301,7 +1348,9 @@ bool FLevelSnapshotsEditorResultsRow::HasChangedChildren() const
 
 bool FLevelSnapshotsEditorResultsRow::ShouldRowBeVisible() const
 {
-	return GetDoesRowMatchSearchTerms() || HasVisibleChildren();
+	const bool bShowUnselectedRows = ResultsViewPtr.IsValid() ? ResultsViewPtr.Pin()->GetShowUnselectedRows() : true;
+	const bool bShouldBeVisibleBasedOnCheckedState = bShowUnselectedRows ? true : GetWidgetCheckedState() != ECheckBoxState::Unchecked;
+	return bShouldBeVisibleBasedOnCheckedState && (GetDoesRowMatchSearchTerms() || HasVisibleChildren());
 }
 
 EVisibility FLevelSnapshotsEditorResultsRow::GetDesiredVisibility() const
@@ -1425,7 +1474,7 @@ void FLevelSnapshotsEditorResultsRow::EvaluateAndSetAllParentGroupCheckedStates(
 				}
 			}
 
-			PinnedParent->WidgetCheckedState = NewWidgetCheckedState;
+			PinnedParent->SetWidgetCheckedState(NewWidgetCheckedState);
 		}
 
 		ParentRow = PinnedParent->GetDirectParentRow();
@@ -1440,10 +1489,18 @@ void SLevelSnapshotsEditorResults::Construct(const FArguments& InArgs, ULevelSna
 	}
 	EditorDataPtr = InEditorData;
 
+	DefaultNameText = LOCTEXT("LevelSnapshots", "Level Snapshots");
+
 	OnActiveSnapshotChangedHandle = InEditorData->OnActiveSnapshotChanged.AddSP(this, &SLevelSnapshotsEditorResults::OnSnapshotSelected);
 
 	OnRefreshResultsHandle = InEditorData->OnRefreshResults.AddSP(this, &SLevelSnapshotsEditorResults::RefreshResults);
 
+	OnMapOpenedDelegateHandle = FEditorDelegates::OnMapOpened.AddLambda([this](const FString& FileName, bool bAsTemplate)
+    {
+		FlushMemory(false);
+		UpdateSnapshotInformationText();
+    });
+
 	FMenuBuilder ShowOptionsMenuBuilder = BuildShowOptionsMenu();
 
 	ChildSlot
@@ -1472,7 +1529,7 @@ void SLevelSnapshotsEditorResults::Construct(const FArguments& InArgs, ULevelSna
 				.VAlign(VAlign_Center)
 				[
 					SAssignNew(SelectedSnapshotNamePtr, STextBlock)
-					.Text(LOCTEXT("LevelSnapshots", "Level Snapshots"))
+					.Text(DefaultNameText)
 					.Font(FCoreStyle::GetDefaultFontStyle("Bold", 16))
 				]
 			]	
@@ -1573,7 +1630,12 @@ void SLevelSnapshotsEditorResults::Construct(const FArguments& InArgs, ULevelSna
 						+SVerticalBox::Slot()
 						.AutoHeight()
 						[
-							SNew(SHorizontalBox)
+							SAssignNew(InfoTextBox, SHorizontalBox)
+							.Visibility_Lambda([this]()
+							{
+								return EditorDataPtr.IsValid() && EditorDataPtr->GetSelectedWorld() && EditorDataPtr->GetActiveSnapshot() ? 
+									EVisibility::HitTestInvisible : EVisibility::Hidden;
+							})
 
 							+SHorizontalBox::Slot()
 							.AutoWidth()
@@ -1609,6 +1671,10 @@ void SLevelSnapshotsEditorResults::Construct(const FArguments& InArgs, ULevelSna
 						.Padding(5.f, 2.f)
 						[
 							SNew(SButton)
+							.IsEnabled_Lambda([this]()
+							{
+								return EditorDataPtr.IsValid() && EditorDataPtr->GetSelectedWorld() && EditorDataPtr->GetActiveSnapshot();
+							})
 							.ButtonStyle(FEditorStyle::Get(), "FlatButton.Success")
 							.ForegroundColor(FSlateColor::UseForeground())
 							.OnClicked_Raw(this, &SLevelSnapshotsEditorResults::OnClickApplyToWorld)
@@ -1638,6 +1704,7 @@ SLevelSnapshotsEditorResults::~SLevelSnapshotsEditorResults()
 	
 	OnActiveSnapshotChangedHandle.Reset();
 	OnRefreshResultsHandle.Reset();
+	OnMapOpenedDelegateHandle.Reset();
 	
 	ResultsSearchBoxPtr.Reset();
 	ResultsBoxContainerPtr.Reset();
@@ -1652,7 +1719,7 @@ SLevelSnapshotsEditorResults::~SLevelSnapshotsEditorResults()
 
 	DummyRow.Reset();
 
-	FlushMemory();
+	FlushMemory(false);
 
 	TreeViewPtr.Reset();
 }
@@ -1662,30 +1729,30 @@ FMenuBuilder SLevelSnapshotsEditorResults::BuildShowOptionsMenu()
 	FMenuBuilder ShowOptionsMenuBuilder = FMenuBuilder(true, nullptr);
 
 	ShowOptionsMenuBuilder.AddMenuEntry(
-		LOCTEXT("ShowFilteredActors", "Show Filtered Actors"),
-		LOCTEXT("ShowFilteredActors_Tooltip", "If false, only displays actors which have passed the applied filter. Refresh Results to apply."),
+		LOCTEXT("ShowFilteredRows", "Show Filtered Rows"),
+		LOCTEXT("ShowFilteredRows_Tooltip", "If false, only displays rows which have passed the applied filter. Refresh Results to apply."),
 		FSlateIcon(),
 		FUIAction(
 			FExecuteAction::CreateLambda([this]() {
-				SetShowFilteredActors(!GetShowFilteredActors());
+				SetShowFilteredRows(!GetShowFilteredRows());
 				}),
 			FCanExecuteAction(),
-			FIsActionChecked::CreateSP(this, &SLevelSnapshotsEditorResults::GetShowFilteredActors)
+			FIsActionChecked::CreateSP(this, &SLevelSnapshotsEditorResults::GetShowFilteredRows)
 		),
 		NAME_None,
 		EUserInterfaceActionType::ToggleButton
 	);
 
 	ShowOptionsMenuBuilder.AddMenuEntry(
-		LOCTEXT("ShowUnselectedActors", "Show Unselected Actors"),
-		LOCTEXT("ShowUnselectedActors_Tooltip", "If false, unselected actors will be hidden from view."),
+		LOCTEXT("ShowUnselectedRows", "Show Unselected Rows"),
+		LOCTEXT("ShowUnselectedRows_Tooltip", "If false, unselected rows will be hidden from view."),
 		FSlateIcon(),
 		FUIAction(
 			FExecuteAction::CreateLambda([this]() {
-				SetShowUnselectedActors(!GetShowUnselectedActors());
+				SetShowUnselectedRows(!GetShowUnselectedRows());
 				}),
 			FCanExecuteAction(),
-			FIsActionChecked::CreateSP(this, &SLevelSnapshotsEditorResults::GetShowUnselectedActors)
+			FIsActionChecked::CreateSP(this, &SLevelSnapshotsEditorResults::GetShowUnselectedRows)
 		),
 		NAME_None,
 		EUserInterfaceActionType::ToggleButton
@@ -1707,34 +1774,44 @@ FMenuBuilder SLevelSnapshotsEditorResults::BuildShowOptionsMenu()
 	return ShowOptionsMenuBuilder;
 }
 
-void SLevelSnapshotsEditorResults::SetShowFilteredActors(const bool bNewSetting)
+void SLevelSnapshotsEditorResults::SetShowFilteredRows(const bool bNewSetting)
 {
 	bShowFilteredActors = bNewSetting;
 }
 
-void SLevelSnapshotsEditorResults::SetShowUnselectedActors(const bool bNewSetting)
+void SLevelSnapshotsEditorResults::SetShowUnselectedRows(const bool bNewSetting)
 {
 	bShowUnselectedActors = bNewSetting;
 }
 
-bool SLevelSnapshotsEditorResults::GetShowFilteredActors() const
+bool SLevelSnapshotsEditorResults::GetShowFilteredRows() const
 {
 	return bShowFilteredActors;
 }
 
-bool SLevelSnapshotsEditorResults::GetShowUnselectedActors() const
+bool SLevelSnapshotsEditorResults::GetShowUnselectedRows() const
 {
 	return bShowUnselectedActors;
 }
 
-void SLevelSnapshotsEditorResults::FlushMemory()
+void SLevelSnapshotsEditorResults::FlushMemory(const bool bShouldKeepMemoryAllocated)
 {
-	TreeViewRootHeaderObjects.Empty();
-	TreeViewModifiedActorGroupObjects.Empty();
-	TreeViewAddedActorGroupObjects.Empty();
-	TreeViewRemovedActorGroupObjects.Empty();
+	if (bShouldKeepMemoryAllocated)
+	{
+		TreeViewRootHeaderObjects.Reset();
+		TreeViewModifiedActorGroupObjects.Reset();
+		TreeViewAddedActorGroupObjects.Reset();
+		TreeViewRemovedActorGroupObjects.Reset();
+	}
+	else
+	{
+		TreeViewRootHeaderObjects.Empty();
+		TreeViewModifiedActorGroupObjects.Empty();
+		TreeViewAddedActorGroupObjects.Empty();
+		TreeViewRemovedActorGroupObjects.Empty();
+	}
 
-	CleanUpGenerators();
+	CleanUpGenerators(bShouldKeepMemoryAllocated);
 }
 
 TOptional<ULevelSnapshot*> SLevelSnapshotsEditorResults::GetSelectedLevelSnapshot() const
@@ -1744,29 +1821,17 @@ TOptional<ULevelSnapshot*> SLevelSnapshotsEditorResults::GetSelectedLevelSnapsho
 
 void SLevelSnapshotsEditorResults::OnSnapshotSelected(const TOptional<ULevelSnapshot*>& InLevelSnapshot)
 {	
-	if (InLevelSnapshot.IsSet())
+	if (InLevelSnapshot.IsSet() && InLevelSnapshot.GetValue())
 	{
-		ULevelSnapshot* Snapshot = InLevelSnapshot.GetValue();
-
-		if (SelectedSnapshotNamePtr.IsValid())
-		{
-			FName SnapshotName = Snapshot->GetSnapshotName();
-			if (SnapshotName == "None")
-			{
-				// If a bespoke snapshot name is not provided then we'll use the asset name. Usually these names are the same.
-				SnapshotName = Snapshot->GetFName();
-			}
-			
-			SelectedSnapshotNamePtr->SetText(FText::FromName(SnapshotName));
-		}
+		UpdateSnapshotNameText(InLevelSnapshot);
 		
-		GenerateTreeView();
+		GenerateTreeView(true);
 	}
 }
 
 void SLevelSnapshotsEditorResults::RefreshResults()
 {
-	GenerateTreeView();
+	GenerateTreeView(false);
 }
 
 FReply SLevelSnapshotsEditorResults::OnClickApplyToWorld()
@@ -1802,6 +1867,30 @@ FReply SLevelSnapshotsEditorResults::OnClickApplyToWorld()
 	return FReply::Handled();
 }
 
+void SLevelSnapshotsEditorResults::UpdateSnapshotNameText(const TOptional<ULevelSnapshot*>& InLevelSnapshot) const
+{
+	if (SelectedSnapshotNamePtr.IsValid())
+	{
+		FText SnapshotName = DefaultNameText;
+		
+		if (InLevelSnapshot.IsSet())
+		{
+			ULevelSnapshot* Snapshot = InLevelSnapshot.GetValue();
+			
+			FName ReturnedName = Snapshot->GetSnapshotName();
+			if (ReturnedName == "None")
+			{
+				// If a bespoke snapshot name is not provided then we'll use the asset name. Usually these names are the same.
+				ReturnedName = Snapshot->GetFName();
+			}
+
+			SnapshotName = FText::FromName(ReturnedName);
+		}
+			
+		SelectedSnapshotNamePtr->SetText(SnapshotName);
+	}
+}
+
 void SLevelSnapshotsEditorResults::UpdateSnapshotInformationText()
 {
 	check(EditorDataPtr.IsValid());
@@ -1867,8 +1956,15 @@ void SLevelSnapshotsEditorResults::BuildSelectionSetFromSelectedPropertiesInEach
 				return;
 			}
 
+			const FPropertySelection* PropertySelection = SelectionMap.GetSelectedProperties(Group->GetWorldObject());
+
+			if (!PropertySelection)
+			{
+				return;
+			}
+
 			// Make a copy of the property selection. If a node is unchecked, we'll remove it from the copy.
-			FPropertySelection CheckedNodeFieldPaths = *SelectionMap.GetSelectedProperties(Group->GetWorldObject());
+			FPropertySelection CheckedNodeFieldPaths = *PropertySelection;
 			SelectionMap.RemoveObjectPropertiesFromMap(Group->GetWorldObject());
 
 			TArray<FLevelSnapshotsEditorResultsRowPtr> UncheckedChildPropertyNodes;
@@ -2022,7 +2118,7 @@ TWeakPtr<FRowGeneratorInfo> SLevelSnapshotsEditorResults::RegisterRowGenerator(
 	return NewGeneratorInfo;
 }
 
-void SLevelSnapshotsEditorResults::CleanUpGenerators()
+void SLevelSnapshotsEditorResults::CleanUpGenerators(const bool bShouldKeepMemoryAllocated)
 {
 	for (TSharedPtr<FRowGeneratorInfo>& Generator : RegisteredRowGenerators)
 	{		
@@ -2033,7 +2129,70 @@ void SLevelSnapshotsEditorResults::CleanUpGenerators()
 		}
 	}
 
-	RegisteredRowGenerators.Empty();
+	if (bShouldKeepMemoryAllocated)
+	{
+		RegisteredRowGenerators.Reset();
+	}
+	else
+	{
+		RegisteredRowGenerators.Empty();
+	}
+}
+
+bool SLevelSnapshotsEditorResults::FindRowStateMemoryByPath(const FString& InPath, FLevelSnapshotsEditorResultsRowStateMemory& OutRowStateMemory)
+{
+	if (RowStateMemory.Num())
+	{
+		const AlgoImpl::TRangePointerType<TSet<TSharedPtr<FLevelSnapshotsEditorResultsRowStateMemory>>>::Type& FindResult =
+			Algo::FindByPredicate(RowStateMemory, [&InPath](const TSharedPtr<FLevelSnapshotsEditorResultsRowStateMemory>& RowStateMemory)
+			{
+				return RowStateMemory->PathToRow.Equals(InPath);
+			});
+
+		if (FindResult)
+		{
+			OutRowStateMemory = *FindResult->Get();
+			return true;
+		}
+	}
+
+	return false;
+}
+
+void SLevelSnapshotsEditorResults::AddRowStateToRowStateMemory(
+	const TSharedPtr<FLevelSnapshotsEditorResultsRowStateMemory> InRowStateMemory)
+{
+	RowStateMemory.Add(InRowStateMemory);
+}
+
+void SLevelSnapshotsEditorResults::GenerateRowStateMemoryRecursively()
+{
+	struct Local
+	{
+		static void GenerateRowStateMemory(const TSharedRef<SLevelSnapshotsEditorResults>& InResultsView, const TSharedPtr<FLevelSnapshotsEditorResultsRow>& Row)
+		{
+			if (!Row.IsValid())
+			{
+				return;
+			}
+			
+			InResultsView->AddRowStateToRowStateMemory(
+				MakeShared<FLevelSnapshotsEditorResultsRowStateMemory>(Row->GetOrGenerateRowPath(), Row->GetIsTreeViewItemExpanded(), Row->GetWidgetCheckedState()));
+
+			for (const TSharedPtr<FLevelSnapshotsEditorResultsRow>& ChildRow : Row->GetChildRows())
+			{
+				GenerateRowStateMemory(InResultsView, ChildRow);
+			}
+		}
+	};
+	
+	// Reset() rather than Empty() because it's likely we'll use a similar amount of memory
+	RowStateMemory.Reset();
+
+	for (const TSharedPtr<FLevelSnapshotsEditorResultsRow>& Row : TreeViewRootHeaderObjects)
+	{
+		Local::GenerateRowStateMemory(SharedThis(this), Row);
+	}
 }
 
 FLevelSnapshotsEditorResultsRowPtr& SLevelSnapshotsEditorResults::GetOrCreateDummyRow()
@@ -2042,13 +2201,13 @@ FLevelSnapshotsEditorResultsRowPtr& SLevelSnapshotsEditorResults::GetOrCreateDum
 	{
 		DummyRow = MakeShared<FLevelSnapshotsEditorResultsRow>(
 			FLevelSnapshotsEditorResultsRow(
-				FText::FromString("Dummy"), FLevelSnapshotsEditorResultsRow::None, ECheckBoxState::Undetermined, nullptr));
+				FText::FromString("Dummy"), FLevelSnapshotsEditorResultsRow::None, ECheckBoxState::Undetermined, SharedThis(this)));
 	}
 	
 	return DummyRow;
 }
 
-void SLevelSnapshotsEditorResults::GenerateTreeView()
+void SLevelSnapshotsEditorResults::GenerateTreeView(const bool bSnapshotHasChanged)
 {	
 	if (!ensure(EditorDataPtr.IsValid()) || !ensure(TreeViewPtr.IsValid()))
 	{
@@ -2056,8 +2215,17 @@ void SLevelSnapshotsEditorResults::GenerateTreeView()
 	}
 
 	DECLARE_SCOPE_CYCLE_COUNTER(TEXT("GenerateTreeView"), STAT_GenerateTreeView, STATGROUP_LevelSnapshots);
+
+	if (bSnapshotHasChanged)
+	{
+		RowStateMemory.Empty();
+	}
+	else
+	{
+		GenerateRowStateMemoryRecursively();
+	}
 	
-	FlushMemory();
+	FlushMemory(!bSnapshotHasChanged);
 	
 	UFilteredResults* FilteredResults = EditorDataPtr->GetFilterResults(); 
 	const TWeakObjectPtr<ULevelSnapshotFilter>& UserFilters = FilteredResults->GetUserFilters();
@@ -2072,13 +2240,12 @@ void SLevelSnapshotsEditorResults::GenerateTreeView()
 	if (FilterListData.GetModifiedFilteredActors().Num())
 	{
 		FLevelSnapshotsEditorResultsRowPtr ModifiedActorsHeader = MakeShared<FLevelSnapshotsEditorResultsRow>(
-			FLevelSnapshotsEditorResultsRow(FText::GetEmpty(), FLevelSnapshotsEditorResultsRow::TreeViewHeader, ECheckBoxState::Checked));
+			FLevelSnapshotsEditorResultsRow(FText::GetEmpty(), FLevelSnapshotsEditorResultsRow::TreeViewHeader, ECheckBoxState::Checked, SharedThis(this)));
 		ModifiedActorsHeader->InitHeaderRow(FLevelSnapshotsEditorResultsRow::ELevelSnapshotsEditorResultsTreeViewHeaderType::HeaderType_ModifiedActors, {
 			LOCTEXT("ColumnName_Actor", "Actor"),
 			LOCTEXT("ColumnName_ValueToRestore", "Value to Restore"),
 			LOCTEXT("ColumnName_CurrentValue", "Current Value")
-			},
-			SharedThis(this));
+			});
 
 		if (GenerateTreeViewChildren_ModifiedActors(ModifiedActorsHeader, UserFilters.Get()))
 		{
@@ -2093,9 +2260,9 @@ void SLevelSnapshotsEditorResults::GenerateTreeView()
 	if (FilterListData.GetFilteredAddedWorldActors().Num())
 	{
 		FLevelSnapshotsEditorResultsRowPtr AddedActorsHeader = MakeShared<FLevelSnapshotsEditorResultsRow>(
-			FLevelSnapshotsEditorResultsRow(FText::GetEmpty(), FLevelSnapshotsEditorResultsRow::TreeViewHeader, ECheckBoxState::Checked));
+			FLevelSnapshotsEditorResultsRow(FText::GetEmpty(), FLevelSnapshotsEditorResultsRow::TreeViewHeader, ECheckBoxState::Checked, SharedThis(this)));
 		AddedActorsHeader->InitHeaderRow(FLevelSnapshotsEditorResultsRow::ELevelSnapshotsEditorResultsTreeViewHeaderType::HeaderType_AddedActors, 
-			{ LOCTEXT("ColumnName_AddedActors", "Added Actors") }, SharedThis(this));
+			{ LOCTEXT("ColumnName_AddedActors", "Added Actors") });
 
 		if (GenerateTreeViewChildren_AddedActors(AddedActorsHeader))
 		{
@@ -2110,9 +2277,9 @@ void SLevelSnapshotsEditorResults::GenerateTreeView()
 	if (FilterListData.GetFilteredRemovedOriginalActorPaths().Num())
 	{
 		FLevelSnapshotsEditorResultsRowPtr RemovedActorsHeader = MakeShared<FLevelSnapshotsEditorResultsRow>(
-			FLevelSnapshotsEditorResultsRow(FText::GetEmpty(), FLevelSnapshotsEditorResultsRow::TreeViewHeader, ECheckBoxState::Checked));
+			FLevelSnapshotsEditorResultsRow(FText::GetEmpty(), FLevelSnapshotsEditorResultsRow::TreeViewHeader, ECheckBoxState::Checked, SharedThis(this)));
 		RemovedActorsHeader->InitHeaderRow(FLevelSnapshotsEditorResultsRow::ELevelSnapshotsEditorResultsTreeViewHeaderType::HeaderType_RemovedActors, 
-			{ LOCTEXT("ColumnName_RemovedActors", "Removed Actors") }, SharedThis(this));
+			{ LOCTEXT("ColumnName_RemovedActors", "Removed Actors") });
 
 		if (GenerateTreeViewChildren_RemovedActors(RemovedActorsHeader))
 		{
@@ -2184,9 +2351,10 @@ bool SLevelSnapshotsEditorResults::GenerateTreeViewChildren_ModifiedActors(FLeve
 
 		// Create group
 		FLevelSnapshotsEditorResultsRowPtr NewActorGroup = MakeShared<FLevelSnapshotsEditorResultsRow>(
-			FLevelSnapshotsEditorResultsRow(FText::FromString(ActorName), FLevelSnapshotsEditorResultsRow::ActorGroup, ECheckBoxState::Checked, ModifiedActorsHeader));
+			FLevelSnapshotsEditorResultsRow(
+				FText::FromString(ActorName), FLevelSnapshotsEditorResultsRow::ActorGroup, ECheckBoxState::Checked, SharedThis(this), ModifiedActorsHeader));
 		
-		NewActorGroup->InitActorRow(WeakSnapshotActor.IsValid() ? WeakSnapshotActor.Get() : nullptr, WorldActor, SharedThis(this));
+		NewActorGroup->InitActorRow(WeakSnapshotActor.IsValid() ? WeakSnapshotActor.Get() : nullptr, WorldActor);
 
 		ModifiedActorsHeader->AddToChildRows(NewActorGroup);
 		
@@ -2234,7 +2402,8 @@ bool SLevelSnapshotsEditorResults::GenerateTreeViewChildren_AddedActors(FLevelSn
 		
 		// Create group
 		FLevelSnapshotsEditorResultsRowPtr NewActorRow = MakeShared<FLevelSnapshotsEditorResultsRow>(
-			FLevelSnapshotsEditorResultsRow(FText::FromString(Actor.Get()->GetName()), FLevelSnapshotsEditorResultsRow::AddedActor, ECheckBoxState::Checked, AddedActorsHeader));
+			FLevelSnapshotsEditorResultsRow(
+				FText::FromString(Actor.Get()->GetName()), FLevelSnapshotsEditorResultsRow::AddedActor, ECheckBoxState::Checked, SharedThis(this), AddedActorsHeader));
 		NewActorRow->InitAddedActorRow(Actor.Get());
 
 		AddedActorsHeader->AddToChildRows(NewActorRow);
@@ -2260,8 +2429,9 @@ bool SLevelSnapshotsEditorResults::GenerateTreeViewChildren_RemovedActors(FLevel
 		
 		// Create group
 		FLevelSnapshotsEditorResultsRowPtr NewActorRow = MakeShared<FLevelSnapshotsEditorResultsRow>(
-			FLevelSnapshotsEditorResultsRow(FText::FromString(ActorName), FLevelSnapshotsEditorResultsRow::RemovedActor, ECheckBoxState::Checked, RemovedActorsHeader));
-		NewActorRow->InitRemovedActorRow(ActorPath, SharedThis(this));
+			FLevelSnapshotsEditorResultsRow(
+				FText::FromString(ActorName), FLevelSnapshotsEditorResultsRow::RemovedActor, ECheckBoxState::Checked, SharedThis(this), RemovedActorsHeader));
+		NewActorRow->InitRemovedActorRow(ActorPath);
 
 		RemovedActorsHeader->AddToChildRows(NewActorRow);
 
@@ -2344,6 +2514,14 @@ bool SLevelSnapshotsEditorResults::DoesTreeViewHaveVisibleChildren() const
 	return false;
 }
 
+void SLevelSnapshotsEditorResults::SetTreeViewItemExpanded(const TSharedPtr<FLevelSnapshotsEditorResultsRow>& RowToExpand, const bool bNewExpansion) const
+{
+	if (TreeViewPtr.IsValid())
+	{
+		TreeViewPtr->SetItemExpansion(RowToExpand, bNewExpansion);
+	}
+}
+
 void SLevelSnapshotsEditorResults::OnGetRowChildren(FLevelSnapshotsEditorResultsRowPtr Row, TArray<FLevelSnapshotsEditorResultsRowPtr>& OutChildren)
 {
 	if (Row.IsValid())
diff --git a/Engine/Plugins/VirtualProduction/LevelSnapshots/Source/LevelSnapshotsEditor/Private/Views/Results/SLevelSnapshotsEditorResults.h b/Engine/Plugins/VirtualProduction/LevelSnapshots/Source/LevelSnapshotsEditor/Private/Views/Results/SLevelSnapshotsEditorResults.h
--- a/Engine/Plugins/VirtualProduction/LevelSnapshots/Source/LevelSnapshotsEditor/Private/Views/Results/SLevelSnapshotsEditorResults.h
+++ b/Engine/Plugins/VirtualProduction/LevelSnapshots/Source/LevelSnapshotsEditor/Private/Views/Results/SLevelSnapshotsEditorResults.h
@@ -4,6 +4,7 @@
 
 #include "Widgets/Input/SCheckBox.h"
 #include "Widgets/SCompoundWidget.h"
+#include "Widgets/SBoxPanel.h"
 #include "Widgets/Views/STreeView.h"
 
 #include "FilterListData.h"
@@ -48,6 +49,25 @@ struct FLevelSnapshotsEditorResultsSplitterManager
 	float WorldObjectPropertyColumnWidth = 1.0f;
 };
 
+struct FLevelSnapshotsEditorResultsRowStateMemory
+{
+	FLevelSnapshotsEditorResultsRowStateMemory()
+	{
+		bIsExpanded = false;
+		WidgetCheckedState = ECheckBoxState::Checked;
+	};
+	
+	FLevelSnapshotsEditorResultsRowStateMemory(const FString& InPathToRow, const bool bNewIsExpanded, const ECheckBoxState NewWidgetCheckedState)
+		: PathToRow(InPathToRow)
+		, bIsExpanded(bNewIsExpanded)
+		, WidgetCheckedState(NewWidgetCheckedState)
+	{};
+	
+	FString PathToRow;
+	bool bIsExpanded;
+	ECheckBoxState WidgetCheckedState;
+};
+
 struct FRowGeneratorInfo
 {
 	~FRowGeneratorInfo()
@@ -120,25 +140,27 @@ struct FLevelSnapshotsEditorResultsRow final : TSharedFromThis<FLevelSnapshotsEd
 	void FlushReferences();
 	
 	FLevelSnapshotsEditorResultsRow(const FText InDisplayName, const ELevelSnapshotsEditorResultsRowType InRowType, const ECheckBoxState StartingWidgetCheckboxState, 
-		const TWeakPtr<FLevelSnapshotsEditorResultsRow>& = nullptr);
+		const TWeakPtr<SLevelSnapshotsEditorResults>& InResultsView, const TWeakPtr<FLevelSnapshotsEditorResultsRow>& InDirectParentRow = nullptr);
 
 	void InitHeaderRow(
-		const ELevelSnapshotsEditorResultsTreeViewHeaderType InHeaderType, const TArray<FText>& InColumns, const TWeakPtr<SLevelSnapshotsEditorResults>& InResultsView);
+		const ELevelSnapshotsEditorResultsTreeViewHeaderType InHeaderType, const TArray<FText>& InColumns);
 	
 	void InitAddedActorRow(AActor* InAddedActor);
-	void InitRemovedActorRow(const FSoftObjectPath& InRemovedActorPath, const TWeakPtr<SLevelSnapshotsEditorResults>& InResultsView);
+	void InitRemovedActorRow(const FSoftObjectPath& InRemovedActorPath);
 	
-	void InitActorRow(AActor* InSnapshotActor, AActor* InWorldActor, const TWeakPtr<SLevelSnapshotsEditorResults>& InResultsView);
+	void InitActorRow(AActor* InSnapshotActor, AActor* InWorldActor);
 	void InitObjectRow(
 		UObject* InSnapshotObject, UObject* InWorldObject,
 		const TWeakPtr<FRowGeneratorInfo>& InSnapshotRowGenerator,
-		const TWeakPtr<FRowGeneratorInfo>& InWorldRowGenerator,
-		const TWeakPtr<SLevelSnapshotsEditorResults>& InResultsView);
+		const TWeakPtr<FRowGeneratorInfo>& InWorldRowGenerator);
 
 	void InitPropertyRow(
 		const TWeakPtr<FLevelSnapshotsEditorResultsRow>& InContainingObjectGroup,
 		const TSharedPtr<FPropertyHandleHierarchy>& InSnapshotHierarchy, const TSharedPtr<FPropertyHandleHierarchy>& InWorldHandleHierarchy,
-		const bool bNewIsCounterpartValueSame = false);
+		const bool bNewIsCounterpartValueSame);
+
+	void ApplyRowStateMemoryIfAvailable();
+	const FString& GetOrGenerateRowPath();
 
 	void GenerateActorGroupChildren(FPropertySelectionMap& PropertySelectionMap);
 
@@ -208,10 +230,10 @@ struct FLevelSnapshotsEditorResultsRow final : TSharedFromThis<FLevelSnapshotsEd
 	void SetIsCounterpartValueSame(const bool bIsValueSame);
 
 	ECheckBoxState GetWidgetCheckedState() const;
-	void SetWidgetCheckedState(const ECheckBoxState NewState, const bool bUserClicked = false);
+	void SetWidgetCheckedState(const ECheckBoxState NewState, const bool bShouldUpdateHierarchyCheckedStates = false);
 
 	bool GetIsNodeChecked() const;
-	void SetIsNodeChecked(const bool bNewChecked);
+	void SetIsNodeChecked(const bool bNewChecked, const bool bShouldUpdateHierarchyCheckedStates = false);
 
 	/* Hierarchy utilities */
 
@@ -254,6 +276,9 @@ private:
 	FString CachedSearchTerms;
 	bool bDoesRowMatchSearchTerms = true;
 
+	/* This is a breadcrumb trail of display names used to find or store the state of the row. */
+	FString RowPath;
+
 	/* Returns a string of searchable keywords such as object names, property names, paths or anything else associated with the row that might be useful to search for. */
 	const FString& GetOrCacheSearchTerms();
 
@@ -306,13 +331,13 @@ public:
 
 	FMenuBuilder BuildShowOptionsMenu();
 
-	void SetShowFilteredActors(const bool bNewSetting);
-	void SetShowUnselectedActors(const bool bNewSetting);
+	void SetShowFilteredRows(const bool bNewSetting);
+	void SetShowUnselectedRows(const bool bNewSetting);
 
-	bool GetShowFilteredActors() const;
-	bool GetShowUnselectedActors() const;
+	bool GetShowFilteredRows() const;
+	bool GetShowUnselectedRows() const;
 	
-	void FlushMemory();
+	void FlushMemory(const bool bShouldKeepMemoryAllocated);
 
 	TOptional<ULevelSnapshot*> GetSelectedLevelSnapshot() const;
 
@@ -320,6 +345,7 @@ public:
 	void RefreshResults();
 	FReply OnClickApplyToWorld();
 
+	void UpdateSnapshotNameText(const TOptional<ULevelSnapshot*>& InLevelSnapshot) const;
 	void UpdateSnapshotInformationText();
 
 	/* This method builds a selection set of all visible and checked properties to pass back to apply to the world. */
@@ -331,6 +357,8 @@ public:
 
 	bool DoesTreeViewHaveVisibleChildren() const;
 
+	void SetTreeViewItemExpanded(const TSharedPtr<FLevelSnapshotsEditorResultsRow>& RowToExpand, const bool bNewExpansion) const;
+
 	ULevelSnapshotsEditorData* GetEditorDataPtr() const;
 
 	// Row Generator Management
@@ -343,7 +371,11 @@ public:
 		const TWeakPtr<FLevelSnapshotsEditorResultsRow>& InBoundObject, const ELevelSnapshotsObjectType InGeneratorType,
 		FPropertyEditorModule& PropertyEditorModule);
 
-	void CleanUpGenerators();
+	void CleanUpGenerators(const bool bShouldKeepMemoryAllocated);
+
+	bool FindRowStateMemoryByPath(const FString& InPath, FLevelSnapshotsEditorResultsRowStateMemory& OutRowStateMemory);
+	void AddRowStateToRowStateMemory(const TSharedPtr<FLevelSnapshotsEditorResultsRowStateMemory> InRowStateMemory);
+	void GenerateRowStateMemoryRecursively();
 
 private:
 
@@ -353,9 +385,12 @@ private:
 	TSharedPtr<STextBlock> SelectedSnapshotNamePtr;
 
 	// Snapshot Information Text
+	TSharedPtr<SHorizontalBox> InfoTextBox;
 	TSharedPtr<STextBlock> SelectedActorCountText;
 	TSharedPtr<STextBlock> TotalActorCountText;
 	TSharedPtr<STextBlock> MiscActorCountText;
+
+	FText DefaultNameText;
 		
 	FReply SetAllActorGroupsCollapsed();
 
@@ -365,6 +400,7 @@ private:
 
 	FDelegateHandle OnActiveSnapshotChangedHandle;
 	FDelegateHandle OnRefreshResultsHandle;
+	FDelegateHandle OnMapOpenedDelegateHandle;
 
 	TWeakObjectPtr<ULevelSnapshotsEditorData> EditorDataPtr;
 
@@ -381,7 +417,7 @@ private:
 
 	//  Tree View Implementation
 
-	void GenerateTreeView();
+	void GenerateTreeView(const bool bSnapshotHasChanged);
 	bool GenerateTreeViewChildren_ModifiedActors(FLevelSnapshotsEditorResultsRowPtr ModifiedActorsHeader, ULevelSnapshotFilter* UserFilters);
 	bool GenerateTreeViewChildren_AddedActors(FLevelSnapshotsEditorResultsRowPtr AddedActorsHeader);
 	bool GenerateTreeViewChildren_RemovedActors(FLevelSnapshotsEditorResultsRowPtr RemovedActorsHeader);
@@ -403,6 +439,9 @@ private:
 
 	/* The results view should be the sole manager of the RowGenerators' lifetimes */
 	TSet<TSharedPtr<FRowGeneratorInfo>> RegisteredRowGenerators;
+
+	/* This list remembers the checked and expansion states of previously created rows so these states can be recreated on refresh. */
+	TSet<TSharedPtr<FLevelSnapshotsEditorResultsRowStateMemory>> RowStateMemory;
 };
 
 class SLevelSnapshotsEditorResultsRow : public SCompoundWidget
