Fix declaration shadowing issue

diff --git a/Engine/Plugins/VirtualProduction/LevelSnapshots/Source/LevelSnapshotsEditor/Private/Views/Results/SLevelSnapshotsEditorResults.cpp b/Engine/Plugins/VirtualProduction/LevelSnapshots/Source/LevelSnapshotsEditor/Private/Views/Results/SLevelSnapshotsEditorResults.cpp
--- a/Engine/Plugins/VirtualProduction/LevelSnapshots/Source/LevelSnapshotsEditor/Private/Views/Results/SLevelSnapshotsEditorResults.cpp
+++ b/Engine/Plugins/VirtualProduction/LevelSnapshots/Source/LevelSnapshotsEditor/Private/Views/Results/SLevelSnapshotsEditorResults.cpp
@@ -2144,9 +2144,9 @@ bool SLevelSnapshotsEditorResults::FindRowStateMemoryByPath(const FString& InPat
 	if (RowStateMemory.Num())
 	{
 		const AlgoImpl::TRangePointerType<TSet<TSharedPtr<FLevelSnapshotsEditorResultsRowStateMemory>>>::Type& FindResult =
-			Algo::FindByPredicate(RowStateMemory, [&InPath](const TSharedPtr<FLevelSnapshotsEditorResultsRowStateMemory>& RowStateMemory)
+			Algo::FindByPredicate(RowStateMemory, [&InPath](const TSharedPtr<FLevelSnapshotsEditorResultsRowStateMemory>& InRowStateMemory)
 			{
-				return RowStateMemory->PathToRow.Equals(InPath);
+				return InRowStateMemory->PathToRow.Equals(InPath);
 			});
 
 		if (FindResult)
