Expose which actors were allowed & disallowed by filters

diff --git a/Engine/Plugins/VirtualProduction/LevelSnapshots/Source/LevelSnapshotsEditor/Private/Data/FilterListData.cpp b/Engine/Plugins/VirtualProduction/LevelSnapshots/Source/LevelSnapshotsEditor/Private/Data/FilterListData.cpp
--- a/Engine/Plugins/VirtualProduction/LevelSnapshots/Source/LevelSnapshotsEditor/Private/Data/FilterListData.cpp
+++ b/Engine/Plugins/VirtualProduction/LevelSnapshots/Source/LevelSnapshotsEditor/Private/Data/FilterListData.cpp
@@ -23,10 +23,10 @@ void FFilterListData::UpdateFilteredList(UWorld* World, ULevelSnapshot* FromSnap
     RelatedSnapshot = FromSnapshot;
 
 	// We expect the number of filtered actors & components to stay roughly the same: retain existing memory
-	ModifiedActorsSelectedProperties.Empty(false);
-	ModifiedFilteredActors.Empty(ModifiedFilteredActors.Num());
-	FilteredRemovedOriginalActorPaths.Empty(FilteredRemovedOriginalActorPaths.Num());
-	FilteredAddedWorldActors.Empty(FilteredAddedWorldActors.Num());
+	ModifiedActorsSelectedProperties_AllowedByFilter.Empty(false);
+	ModifiedWorldActors_AllowedByFilter.Empty(ModifiedWorldActors_AllowedByFilter.Num());
+	RemovedOriginalActorPaths_AllowedByFilter.Empty(RemovedOriginalActorPaths_AllowedByFilter.Num());
+	AddedWorldActors_AllowedByFilter.Empty(AddedWorldActors_AllowedByFilter.Num());
 	
 	FromSnapshot->DiffWorld(
 		World,
@@ -36,20 +36,39 @@ void FFilterListData::UpdateFilteredList(UWorld* World, ULevelSnapshot* FromSnap
 		);
 }
 
-const FPropertySelectionMap& FFilterListData::ApplyFilterToFindSelectedProperties(AActor* WorldActor, ULevelSnapshotFilter* FilterToApply)
+void FFilterListData::ApplyFilterToFindSelectedProperties(AActor* WorldActor, ULevelSnapshotFilter* FilterToApply)
 {
-	if (!ensure(WorldActor && FilterToApply) || ModifiedActorsSelectedProperties.GetSelectedProperties(WorldActor))
+	DECLARE_SCOPE_CYCLE_COUNTER(TEXT("ApplyFilterToFindSelectedProperties"), STAT_ApplyFilterToFindSelectedProperties, STATGROUP_LevelSnapshots);
+	
+	const FPropertySelection* AllowedSelectedProperties = ModifiedActorsSelectedProperties_AllowedByFilter.GetSelectedProperties(WorldActor);
+	const FPropertySelection* DisallowedSelectedProperties = ModifiedActorsSelectedProperties_AllowedByFilter.GetSelectedProperties(WorldActor);
+	
+	if (!ensure(WorldActor && FilterToApply) || AllowedSelectedProperties || DisallowedSelectedProperties)
 	{
-		return ModifiedActorsSelectedProperties;
+		return;
+	}
+	
+	const bool bIsAllowedByFilters = ModifiedWorldActors_AllowedByFilter.Contains(WorldActor);
+	const bool bIsDisallowedByFilters = ModifiedWorldActors_DisallowedByFilter.Contains(WorldActor);
+	if (!ensureMsgf(bIsAllowedByFilters || bIsDisallowedByFilters, TEXT("You have to call UpdateFilteredList first")))
+	{
+		return; 
 	}
-	DECLARE_SCOPE_CYCLE_COUNTER(TEXT("ApplyFilterToFindSelectedProperties"), STAT_ApplyFilterToFindSelectedProperties, STATGROUP_LevelSnapshots);
 
 	const TWeakObjectPtr<AActor> DeserializedActor = GetSnapshotCounterpartFor(WorldActor);
-	if (DeserializedActor.IsValid())
+	if (!ensureMsgf(DeserializedActor.IsValid(), TEXT("For some reason this actor has no snapshot counterpart... Investigate.")))
 	{
-		ULevelSnapshotsFunctionLibrary::ApplyFilterToFindSelectedProperties(RelatedSnapshot, ModifiedActorsSelectedProperties, WorldActor, DeserializedActor.Get(), FilterToApply);
+		return;
+	}
+	
+	if (bIsAllowedByFilters)
+	{
+		ULevelSnapshotsFunctionLibrary::ApplyFilterToFindSelectedProperties(RelatedSnapshot, ModifiedActorsSelectedProperties_AllowedByFilter, WorldActor, DeserializedActor.Get(), FilterToApply);
+	}
+	else
+	{
+		ULevelSnapshotsFunctionLibrary::ApplyFilterToFindSelectedProperties(RelatedSnapshot, ModifiedActorsSelectedProperties_DisallowedByFilter, WorldActor, DeserializedActor.Get(), FilterToApply);
 	}
-	return ModifiedActorsSelectedProperties;
 }
 
 TWeakObjectPtr<AActor> FFilterListData::GetSnapshotCounterpartFor(TWeakObjectPtr<AActor> WorldActor) const
@@ -59,26 +78,6 @@ TWeakObjectPtr<AActor> FFilterListData::GetSnapshotCounterpartFor(TWeakObjectPtr
 		? *DeserializedActor : nullptr;
 }
 
-const FPropertySelectionMap& FFilterListData::GetModifiedActorsSelectedProperties() const
-{
-	return ModifiedActorsSelectedProperties;
-}
-
-const TSet<TWeakObjectPtr<AActor>>& FFilterListData::GetModifiedFilteredActors() const
-{
-	return ModifiedFilteredActors;
-}
-
-const TSet<FSoftObjectPath>& FFilterListData::GetFilteredRemovedOriginalActorPaths() const
-{
-	return FilteredRemovedOriginalActorPaths;
-}
-
-const TSet<TWeakObjectPtr<AActor>>& FFilterListData::GetFilteredAddedWorldActors() const
-{
-	return FilteredAddedWorldActors;
-}
-
 void FFilterListData::HandleActorExistsInWorldAndSnapshot(const FSoftObjectPath& OriginalActorPath, ULevelSnapshotFilter* FilterToApply, FScopedSlowTask* Progress)
 {
 	Progress->EnterProgressFrame();
@@ -104,7 +103,11 @@ void FFilterListData::HandleActorExistsInWorldAndSnapshot(const FSoftObjectPath&
 			const EFilterResult::Type ActorInclusionResult = FilterToApply->IsActorValid(FIsActorValidParams(*DeserializedSnapshotActor, WorldActor));
 			if (EFilterResult::CanInclude(ActorInclusionResult))
 			{
-				ModifiedFilteredActors.Add(WorldActor);
+				ModifiedWorldActors_AllowedByFilter.Add(WorldActor);
+			}
+			else
+			{
+				ModifiedWorldActors_DisallowedByFilter.Add(WorldActor);
 			}
 		}
 	}
@@ -121,9 +124,14 @@ void FFilterListData::HandleActorWasRemovedFromWorld(const FSoftObjectPath& Orig
 			}
 		)
 	);
+	
 	if (EFilterResult::CanInclude(FilterResult))
 	{
-		FilteredRemovedOriginalActorPaths.Add(OriginalActorPath);
+		RemovedOriginalActorPaths_AllowedByFilter.Add(OriginalActorPath);
+	}
+	else
+	{
+		RemovedOriginalActorPaths_DisallowedByFilter.Add(OriginalActorPath);
 	}
 }
 
@@ -132,7 +140,11 @@ void FFilterListData::HandleActorWasAddedToWorld(AActor* WorldActor, ULevelSnaps
 	const EFilterResult::Type FilterResult = FilterToApply->IsAddedActorValid(FIsAddedActorValidParams(WorldActor)); 
 	if (EFilterResult::CanInclude(FilterResult))
 	{
-		FilteredAddedWorldActors.Add(WorldActor);
+		AddedWorldActors_AllowedByFilter.Add(WorldActor);
+	}
+	else
+	{
+		AddedWorldActors_DisallowedByFilter.Add(WorldActor);
 	}
 }
 
diff --git a/Engine/Plugins/VirtualProduction/LevelSnapshots/Source/LevelSnapshotsEditor/Private/Data/FilterListData.h b/Engine/Plugins/VirtualProduction/LevelSnapshots/Source/LevelSnapshotsEditor/Private/Data/FilterListData.h
--- a/Engine/Plugins/VirtualProduction/LevelSnapshots/Source/LevelSnapshotsEditor/Private/Data/FilterListData.h
+++ b/Engine/Plugins/VirtualProduction/LevelSnapshots/Source/LevelSnapshotsEditor/Private/Data/FilterListData.h
@@ -22,23 +22,26 @@ struct FFilterListData
 
 	void UpdateFilteredList(UWorld* World, ULevelSnapshot* FromSnapshot, ULevelSnapshotFilter* FilterToApply);
 	
-	/* Runs IsPropertyValid on all properties of WorldActor.
-	 * 
-	 * Puts the results into ModifiedActorSelectedProperties or UnmodifiedActorSelectedProperties, respectively.
-	 * @return A map containing the selected properties for WorldActor: either ModifiedActorSelectedProperties or UnmodifiedActorSelectedProperties.
+	/**
+	 * Runs IsPropertyValid on all properties of WorldActor.
 	 */
-	const FPropertySelectionMap& ApplyFilterToFindSelectedProperties(AActor* WorldActor, ULevelSnapshotFilter* FilterToApply);
+	void ApplyFilterToFindSelectedProperties(AActor* WorldActor, ULevelSnapshotFilter* FilterToApply);
 
-	/*
+	/**
 	 * If WorldActor is modified, returns the deserialized actor.
 	 * If WorldActor is not modified, returns itself.
 	 */
 	TWeakObjectPtr<AActor> GetSnapshotCounterpartFor(TWeakObjectPtr<AActor> WorldActor) const;
 
-	const FPropertySelectionMap& GetModifiedActorsSelectedProperties() const;
-	const TSet<TWeakObjectPtr<AActor>>& GetModifiedFilteredActors() const;
-	const TSet<FSoftObjectPath>& GetFilteredRemovedOriginalActorPaths() const;
-	const TSet<TWeakObjectPtr<AActor>>& GetFilteredAddedWorldActors() const;
+	const FPropertySelectionMap& GetModifiedActorsSelectedProperties_AllowedByFilter() const { return ModifiedActorsSelectedProperties_AllowedByFilter; }
+	const TSet<TWeakObjectPtr<AActor>>& GetModifiedActors_AllowedByFilter() const { return ModifiedWorldActors_AllowedByFilter; }
+	const TSet<FSoftObjectPath>& GetRemovedOriginalActorPaths_AllowedByFilter() const { return RemovedOriginalActorPaths_AllowedByFilter; }
+	const TSet<TWeakObjectPtr<AActor>>& GetAddedWorldActors_AllowedByFilter() const { return AddedWorldActors_AllowedByFilter; }
+
+	const FPropertySelectionMap& GetModifiedActorsSelectedProperties_DisallowedByFilter() const { return ModifiedActorsSelectedProperties_DisallowedByFilter; }
+	const TSet<TWeakObjectPtr<AActor>>& GetModifiedActors_DisallowedByFilter() const { return ModifiedWorldActors_DisallowedByFilter; }
+	const TSet<FSoftObjectPath>& GetRemovedOriginalActorPaths_DisallowedByFilter() const { return RemovedOriginalActorPaths_DisallowedByFilter; }
+	const TSet<TWeakObjectPtr<AActor>>& GetAddedWorldActors_DisallowedByFilter() const { return AddedWorldActors_DisallowedByFilter; }
 
 private:
 
@@ -49,21 +52,35 @@ private:
 	
 	UPROPERTY()
 	ULevelSnapshot* RelatedSnapshot = nullptr;
+
 	
-	/* Initially empty. Contains the selected properties for actors whose serialized data differs in the selected saved snapshot and the editor world. */
+	/* Selected properties for actors allowed by filters. */
 	UPROPERTY()
-	FPropertySelectionMap ModifiedActorsSelectedProperties;
+	FPropertySelectionMap ModifiedActorsSelectedProperties_AllowedByFilter;
 
 	/* Actors to show in filter results panel when "ShowUnchanged = false". */
 	UPROPERTY()
-	TSet<TWeakObjectPtr<AActor>> ModifiedFilteredActors;
-
+	TSet<TWeakObjectPtr<AActor>> ModifiedWorldActors_AllowedByFilter;
 	/* Actors which existed in snapshot but not in the world. Only contains entries that passed filters. */
 	UPROPERTY()
-	TSet<FSoftObjectPath> FilteredRemovedOriginalActorPaths;
-
+	TSet<FSoftObjectPath> RemovedOriginalActorPaths_AllowedByFilter;
 	/* Actors which existed in the world but not in the snapshot. Only contains entries that passed filters. */
 	UPROPERTY()
-	TSet<TWeakObjectPtr<AActor>> FilteredAddedWorldActors;
+	TSet<TWeakObjectPtr<AActor>> AddedWorldActors_AllowedByFilter;
+
 
+	
+	/* Selected properties for actors disallowed by filters. */
+	UPROPERTY()
+	FPropertySelectionMap ModifiedActorsSelectedProperties_DisallowedByFilter;
+
+	/* Actors to show in filter results panel when "ShowUnchanged = true". */
+	UPROPERTY()
+	TSet<TWeakObjectPtr<AActor>> ModifiedWorldActors_DisallowedByFilter;
+	/* Actors which existed in snapshot but not in the world. Only contains entries that did not pass filters. */
+	UPROPERTY()
+	TSet<FSoftObjectPath> RemovedOriginalActorPaths_DisallowedByFilter;
+	/* Actors which existed in the world but not in the snapshot. Only contains entries that did not pass filters. */
+	UPROPERTY()
+	TSet<TWeakObjectPtr<AActor>> AddedWorldActors_DisallowedByFilter;
 };
\ No newline at end of file
diff --git a/Engine/Plugins/VirtualProduction/LevelSnapshots/Source/LevelSnapshotsEditor/Private/Views/Results/SLevelSnapshotsEditorResults.cpp b/Engine/Plugins/VirtualProduction/LevelSnapshots/Source/LevelSnapshotsEditor/Private/Views/Results/SLevelSnapshotsEditorResults.cpp
--- a/Engine/Plugins/VirtualProduction/LevelSnapshots/Source/LevelSnapshotsEditor/Private/Views/Results/SLevelSnapshotsEditorResults.cpp
+++ b/Engine/Plugins/VirtualProduction/LevelSnapshots/Source/LevelSnapshotsEditor/Private/Views/Results/SLevelSnapshotsEditorResults.cpp
@@ -2015,7 +2015,7 @@ void SLevelSnapshotsEditorResults::BuildSelectionSetFromSelectedPropertiesInEach
 		}
 	};
 
-	FPropertySelectionMap PropertySelectionMap = FilterListData.GetModifiedActorsSelectedProperties();
+	FPropertySelectionMap PropertySelectionMap = FilterListData.GetModifiedActorsSelectedProperties_AllowedByFilter();
 
 	// Modified actors
 	for (const FLevelSnapshotsEditorResultsRowPtr& Group : TreeViewModifiedActorGroupObjects)
@@ -2237,7 +2237,7 @@ void SLevelSnapshotsEditorResults::GenerateTreeView(const bool bSnapshotHasChang
 	SplitterManagerPtr = MakeShared<FLevelSnapshotsEditorResultsSplitterManager>(FLevelSnapshotsEditorResultsSplitterManager());
 
 	// Create root headers
-	if (FilterListData.GetModifiedFilteredActors().Num())
+	if (FilterListData.GetModifiedActors_AllowedByFilter().Num())
 	{
 		FLevelSnapshotsEditorResultsRowPtr ModifiedActorsHeader = MakeShared<FLevelSnapshotsEditorResultsRow>(
 			FLevelSnapshotsEditorResultsRow(FText::GetEmpty(), FLevelSnapshotsEditorResultsRow::TreeViewHeader, ECheckBoxState::Checked, SharedThis(this)));
@@ -2257,7 +2257,7 @@ void SLevelSnapshotsEditorResults::GenerateTreeView(const bool bSnapshotHasChang
 		}
 	}
 
-	if (FilterListData.GetFilteredAddedWorldActors().Num())
+	if (FilterListData.GetAddedWorldActors_AllowedByFilter().Num())
 	{
 		FLevelSnapshotsEditorResultsRowPtr AddedActorsHeader = MakeShared<FLevelSnapshotsEditorResultsRow>(
 			FLevelSnapshotsEditorResultsRow(FText::GetEmpty(), FLevelSnapshotsEditorResultsRow::TreeViewHeader, ECheckBoxState::Checked, SharedThis(this)));
@@ -2274,7 +2274,7 @@ void SLevelSnapshotsEditorResults::GenerateTreeView(const bool bSnapshotHasChang
 		}
 	}
 	
-	if (FilterListData.GetFilteredRemovedOriginalActorPaths().Num())
+	if (FilterListData.GetRemovedOriginalActorPaths_AllowedByFilter().Num())
 	{
 		FLevelSnapshotsEditorResultsRowPtr RemovedActorsHeader = MakeShared<FLevelSnapshotsEditorResultsRow>(
 			FLevelSnapshotsEditorResultsRow(FText::GetEmpty(), FLevelSnapshotsEditorResultsRow::TreeViewHeader, ECheckBoxState::Checked, SharedThis(this)));
@@ -2302,7 +2302,7 @@ bool SLevelSnapshotsEditorResults::GenerateTreeViewChildren_ModifiedActors(FLeve
 {
 	check(ModifiedActorsHeader);
 	
-	const TSet<TWeakObjectPtr<AActor>>& ActorsToConsider = FilterListData.GetModifiedFilteredActors();
+	const TSet<TWeakObjectPtr<AActor>>& ActorsToConsider = FilterListData.GetModifiedActors_AllowedByFilter();
 
 	TSet<FSoftObjectPath> EvaluatedObjects;
 
@@ -2325,7 +2325,7 @@ bool SLevelSnapshotsEditorResults::GenerateTreeViewChildren_ModifiedActors(FLeve
 			continue;
 		}
 
-		const int32 KeyCountBeforeFilter = FilterListData.GetModifiedActorsSelectedProperties().GetKeyCount();
+		const int32 KeyCountBeforeFilter = FilterListData.GetModifiedActorsSelectedProperties_AllowedByFilter().GetKeyCount();
 
 		// Get remaining properties after filter
 		if (UserFilters)
@@ -2333,7 +2333,7 @@ bool SLevelSnapshotsEditorResults::GenerateTreeViewChildren_ModifiedActors(FLeve
 			FilterListData.ApplyFilterToFindSelectedProperties(WorldActor, UserFilters);
 		}
 
-		const FPropertySelectionMap& ModifiedSelectedActors = FilterListData.GetModifiedActorsSelectedProperties();
+		const FPropertySelectionMap& ModifiedSelectedActors = FilterListData.GetModifiedActorsSelectedProperties_AllowedByFilter();
 		const int32 KeyCountAfterFilter = ModifiedSelectedActors.GetKeyCount();
 		const int32 KeyCountDifference = KeyCountAfterFilter - KeyCountBeforeFilter;
 
@@ -2393,7 +2393,7 @@ bool SLevelSnapshotsEditorResults::GenerateTreeViewChildren_AddedActors(FLevelSn
 {
 	check(AddedActorsHeader);
 	
-	for (const TWeakObjectPtr<AActor>& Actor : FilterListData.GetFilteredAddedWorldActors())
+	for (const TWeakObjectPtr<AActor>& Actor : FilterListData.GetAddedWorldActors_AllowedByFilter())
 	{
 		if (!Actor.IsValid())
 		{
@@ -2418,7 +2418,7 @@ bool SLevelSnapshotsEditorResults::GenerateTreeViewChildren_RemovedActors(FLevel
 {
 	check(RemovedActorsHeader);
 	
-	for (const FSoftObjectPath& ActorPath : FilterListData.GetFilteredRemovedOriginalActorPaths())
+	for (const FSoftObjectPath& ActorPath : FilterListData.GetRemovedOriginalActorPaths_AllowedByFilter())
 	{
 		FString ActorName = ActorPath.GetSubPathString().IsEmpty() ? ActorPath.GetAssetName() : ActorPath.GetSubPathString();
 
@@ -2534,7 +2534,7 @@ void SLevelSnapshotsEditorResults::OnGetRowChildren(FLevelSnapshotsEditorResults
 		{
 			if (Row->GetIsTreeViewItemExpanded())
 			{
-				FPropertySelectionMap PropertySelectionMap = FilterListData.GetModifiedActorsSelectedProperties();
+				FPropertySelectionMap PropertySelectionMap = FilterListData.GetModifiedActorsSelectedProperties_AllowedByFilter();
 				Row->GenerateActorGroupChildren(PropertySelectionMap);
 
 				OutChildren = Row->GetChildRows();
