Implement unit tests

diff --git a/Engine/Plugins/VirtualProduction/LevelSnapshots/Source/LevelSnapshots/Private/Data/LevelSnapshot.cpp b/Engine/Plugins/VirtualProduction/LevelSnapshots/Source/LevelSnapshots/Private/Data/LevelSnapshot.cpp
--- a/Engine/Plugins/VirtualProduction/LevelSnapshots/Source/LevelSnapshots/Private/Data/LevelSnapshot.cpp
+++ b/Engine/Plugins/VirtualProduction/LevelSnapshots/Source/LevelSnapshots/Private/Data/LevelSnapshot.cpp
@@ -118,7 +118,8 @@ void ULevelSnapshot::SnapshotWorld(UWorld* TargetWorld)
 		return;
 	}
 
-	if (TargetWorld->WorldType != EWorldType::Editor)
+	if (TargetWorld->WorldType != EWorldType::Editor
+		&& TargetWorld->WorldType != EWorldType::EditorPreview) // To suppor tests in editor preview maps
 	{
 #if !(UE_BUILD_SHIPPING || UE_BUILD_TEST)
 		if (TargetWorld->IsPlayInEditor())
diff --git a/Engine/Plugins/VirtualProduction/LevelSnapshots/Source/LevelSnapshots/Private/LevelSnapshotsFunctionLibrary.cpp b/Engine/Plugins/VirtualProduction/LevelSnapshots/Source/LevelSnapshots/Private/LevelSnapshotsFunctionLibrary.cpp
--- a/Engine/Plugins/VirtualProduction/LevelSnapshots/Source/LevelSnapshots/Private/LevelSnapshotsFunctionLibrary.cpp
+++ b/Engine/Plugins/VirtualProduction/LevelSnapshots/Source/LevelSnapshots/Private/LevelSnapshotsFunctionLibrary.cpp
@@ -105,13 +105,8 @@ ULevelSnapshot* ULevelSnapshotsFunctionLibrary::TakeLevelSnapshot_Internal(const
 		TargetWorld = WorldContextObject->GetWorld();
 	}
 
-	if (TargetWorld)
+	if (!ensure(TargetWorld))
 	{
-		UE_LOG(LogLevelSnapshots, Warning, TEXT("Snapshot taken in World Type - %d"), TargetWorld->WorldType);
-	}
-	else
-	{
-		UE_LOG(LogLevelSnapshots, Warning, TEXT("Snapshot taken with no valid World set"));
 		return nullptr;
 	}
 	
diff --git a/Engine/Plugins/VirtualProduction/LevelSnapshots/Source/LevelSnapshots/Private/LevelSnapshotsModule.cpp b/Engine/Plugins/VirtualProduction/LevelSnapshots/Source/LevelSnapshots/Private/LevelSnapshotsModule.cpp
--- a/Engine/Plugins/VirtualProduction/LevelSnapshots/Source/LevelSnapshots/Private/LevelSnapshotsModule.cpp
+++ b/Engine/Plugins/VirtualProduction/LevelSnapshots/Source/LevelSnapshots/Private/LevelSnapshotsModule.cpp
@@ -8,10 +8,31 @@
 #include "Restorability/StaticMeshCollisionPropertyComparer.h"
 #include "SnapshotRestorability.h"
 
-#include "Components/StaticMeshComponent.h"
 #include "EngineUtils.h"
 #include "Modules/ModuleManager.h"
-#include "Restorability/StaticMeshCollisionPropertyComparer.h"
+
+namespace
+{
+	void AddSoftObjectPathSupport(FLevelSnapshotsModule& Module)
+	{
+		// By default FSnapshotRestorability::IsRestorableProperty requires properties to have the CPF_Edit specifier
+		// FSoftObjectPath does not have this so we need to whitelist its properties
+
+		UStruct* SoftObjectPath = FindObject<UStruct>(nullptr, TEXT("/Script/CoreUObject.SoftObjectPath"));
+		if (!ensureMsgf(SoftObjectPath, TEXT("Investigate why this class could not be found")))
+		{
+			return;
+		}
+
+		TSet<const FProperty*> Properties;
+		for (TFieldIterator<FProperty> FieldIt(SoftObjectPath); FieldIt; ++FieldIt)
+		{
+			Properties.Add(*FieldIt);
+		}
+
+		Module.AddWhitelistedProperties(Properties);
+	}
+}
 
 FLevelSnapshotsModule& FLevelSnapshotsModule::GetInternalModuleInstance()
 {
@@ -30,6 +51,8 @@ void FLevelSnapshotsModule::StartupModule()
 		})
 	);
 	RegisterRestorabilityOverrider(Blacklist);
+	
+	AddSoftObjectPathSupport(*this);
 
 	FStaticMeshCollisionPropertyComparer::Register(*this);
 }
diff --git a/Engine/Plugins/VirtualProduction/LevelSnapshots/Source/LevelSnapshots/Private/Tests/ApplyToWorldTests.cpp b/Engine/Plugins/VirtualProduction/LevelSnapshots/Source/LevelSnapshots/Private/Tests/ApplyToWorldTests.cpp
new file mode 100644
--- /dev/null
+++ b/Engine/Plugins/VirtualProduction/LevelSnapshots/Source/LevelSnapshots/Private/Tests/ApplyToWorldTests.cpp
@@ -0,0 +1,250 @@
+﻿// Copyright Epic Games, Inc. All Rights Reserved.
+
+#include "SnapshotTestRunner.h"
+#include "SnapshotTestActor.h"
+
+#include "Components/PointLightComponent.h"
+#include "Engine/PointLight.h"
+#include "Engine/PostProcessVolume.h"
+#include "Engine/StaticMeshActor.h"
+#include "Misc/AutomationTest.h"
+
+IMPLEMENT_SIMPLE_AUTOMATION_TEST(FRestoreSimpleProperties, "VirtualProduction.LevelSnapshots.Snapshot.RestoreSimpleProperties", (EAutomationTestFlags::ApplicationContextMask | EAutomationTestFlags::EngineFilter));
+bool FRestoreSimpleProperties::RunTest(const FString& Parameters)
+{
+	const FVector StartLocation(100.f, -200.f, 300.f);
+	const FVector EndLocation(500.f, -625.f, 750.f);
+	const FRotator StartRotation(20.f, 30.f, -30.f);
+	const FRotator EndRotation(50.f, -60.f, -40.f);
+	const FVector StartScale(1.f, 1.f, 2.f);
+	const FVector EndScale(2.f, 3.f, -2.f);
+
+	const FPostProcessSettings StartSettings = []()
+	{
+		FPostProcessSettings Result;
+		Result.bOverride_ColorSaturation = true;
+		Result.ColorSaturation = FVector4(0.5f, 0.5f, 0.5f, 0.5f);
+		Result.bOverride_ReflectionsType = true;
+		Result.ReflectionsType = EReflectionsType::ScreenSpace;
+		Result.bOverride_TranslucencyType = true;
+		Result.TranslucencyType = ETranslucencyType::Raster;
+		return Result;
+	}();
+	const FPostProcessSettings EndSettings = []()
+	{
+		FPostProcessSettings Result;
+		Result.bOverride_ColorSaturation = true;
+		Result.ColorSaturation = FVector4(1.5f, 1.5f, 1.5f, 1.5f);
+		Result.bOverride_ReflectionsType = true;
+		Result.ReflectionsType = EReflectionsType::RayTracing;
+		Result.bOverride_TranslucencyType = false;
+		Result.TranslucencyType = ETranslucencyType::RayTracing;
+		return Result;
+	}();
+
+	const float StartRadius = 200.f;
+	const float EndRadius = 500.f;
+	const FColor StartColor = FColor::Red;
+	const FColor EndColor = FColor::Blue;
+	const uint32 bStartCastShadows = true;
+	const uint32 bEndCastShaodws = false;
+	
+	AStaticMeshActor* StaticMesh	= nullptr;
+	APostProcessVolume* PostProcess = nullptr;
+	APointLight* PointLight			= nullptr;
+	
+	FSnapshotTestRunner()
+		.ModifyWorld([&](UWorld* World)
+		{
+			StaticMesh = World->SpawnActor<AStaticMeshActor>(StartLocation, StartRotation);
+			StaticMesh->SetActorScale3D(StartScale);
+
+			PostProcess = World->SpawnActor<APostProcessVolume>();
+			PostProcess->Settings = StartSettings;
+
+			PointLight = World->SpawnActor<APointLight>(StartLocation, StartRotation);
+			PointLight->PointLightComponent->AttenuationRadius = StartRadius;
+			PointLight->PointLightComponent->LightColor = StartColor;
+			PointLight->PointLightComponent->CastShadows = bStartCastShadows;
+		})
+		.TakeSnapshot()
+	
+		.ModifyWorld([&](UWorld* World)
+		{
+			StaticMesh->SetActorLocationAndRotation(EndLocation, EndRotation);
+			StaticMesh->SetActorScale3D(EndScale);
+
+			PostProcess->Settings = EndSettings;
+			PostProcess->SetActorLocationAndRotation(StartLocation, StartRotation);
+
+			PointLight->PointLightComponent->AttenuationRadius = EndRadius;
+			PointLight->PointLightComponent->LightColor = EndColor;
+			PointLight->PointLightComponent->CastShadows = bEndCastShaodws;
+		})
+		.ApplySnapshot()
+
+		.RunTest([&]()
+		{
+			TestEqual("Static Mesh Location", StaticMesh->GetActorLocation(), StartLocation);
+			TestEqual("Static Mesh Rotation", StaticMesh->GetActorRotation(), StartRotation);
+			TestEqual("Static Mesh Scale", StaticMesh->GetActorScale3D(), StartScale);
+
+			TestEqual("Post Process bOverride_ColorSaturation", PostProcess->Settings.bOverride_ColorSaturation, StartSettings.bOverride_ColorSaturation);
+			TestEqual("Post Process ColorSaturation", PostProcess->Settings.ColorSaturation, StartSettings.ColorSaturation);
+			TestEqual("Post Process bOverride_ReflectionsType", PostProcess->Settings.bOverride_ReflectionsType, StartSettings.bOverride_ReflectionsType);
+			TestEqual("Post Process ReflectionsType", PostProcess->Settings.ReflectionsType, StartSettings.ReflectionsType);
+			TestEqual("Post Process bOverride_TranslucencyType", PostProcess->Settings.bOverride_TranslucencyType, StartSettings.bOverride_TranslucencyType);
+			TestEqual("Post Process TranslucencyType", PostProcess->Settings.TranslucencyType, StartSettings.TranslucencyType);
+			
+			TestEqual("Post Process Location", PostProcess->GetActorLocation(), FVector(0.f));
+			TestEqual("Post Process Rotation", PostProcess->GetActorRotation(), FRotator(0.f));
+
+			TestEqual("Point Light Radius", PointLight->PointLightComponent->AttenuationRadius, StartRadius);
+			TestEqual("Point Light Colour", PointLight->PointLightComponent->LightColor, StartColor);
+			TestEqual("Point Light Cast Shadows", PointLight->PointLightComponent->CastShadows, bStartCastShadows);
+		});
+	
+	return true;
+}
+
+IMPLEMENT_SIMPLE_AUTOMATION_TEST(FRestoreReferenceProperties, "VirtualProduction.LevelSnapshots.Snapshot.RestoreReferenceProperties", (EAutomationTestFlags::ApplicationContextMask | EAutomationTestFlags::EngineFilter));
+bool FRestoreReferenceProperties::RunTest(const FString& Parameters)
+{
+	for (TFieldIterator<FProperty> FieldIt(ASnapshotTestActor::StaticClass()); FieldIt; ++FieldIt)
+	{
+		if (const FStructProperty* StructProperty = CastField<FStructProperty>(*FieldIt))
+		{
+			FSoftObjectPath ClassPath(StructProperty->Struct);
+			FString AsString = ClassPath.ToString();
+			UE_LOG(LogTemp, Log, TEXT("%s"), *AsString);
+		}
+	}
+	ASnapshotTestActor* FirstActor = nullptr;
+	ASnapshotTestActor* SecondActor = nullptr;
+	ASnapshotTestActor* PointToSelfActor = nullptr;
+	ASnapshotTestActor* FromWorldToExternal = nullptr;
+	ASnapshotTestActor* FromExternalToWorld = nullptr;
+	ASnapshotTestActor* MaterialAndMesh = nullptr;
+
+	FSnapshotTestRunner()
+		.ModifyWorld([&](UWorld* World)
+		{
+			FirstActor = World->SpawnActor<ASnapshotTestActor>();
+			SecondActor = World->SpawnActor<ASnapshotTestActor>();
+			PointToSelfActor = World->SpawnActor<ASnapshotTestActor>();
+			FromWorldToExternal = World->SpawnActor<ASnapshotTestActor>();
+			FromExternalToWorld = World->SpawnActor<ASnapshotTestActor>();
+			
+			MaterialAndMesh = World->SpawnActor<ASnapshotTestActor>();
+			
+			FirstActor->SetObjectReference(SecondActor);
+			SecondActor->SetObjectReference(FirstActor);
+			PointToSelfActor->SetObjectReference(PointToSelfActor);
+			FromWorldToExternal->SetObjectReference(FirstActor);
+			FromExternalToWorld->SetObjectReference(FromExternalToWorld->CubeMesh);
+			
+			MaterialAndMesh->InstancedMeshComponent->SetStaticMesh(MaterialAndMesh->CubeMesh);
+			MaterialAndMesh->InstancedMeshComponent->SetMaterial(0, MaterialAndMesh->GradientLinearMaterial);
+		})
+		.TakeSnapshot()
+
+		.ModifyWorld([&](UWorld* World)
+		{
+			FirstActor->ClearObjectReferences();
+			// No change to SecondActor
+			PointToSelfActor->SetObjectReference(FirstActor);
+			FromWorldToExternal->SetObjectReference(FromWorldToExternal->CubeMesh);
+			FromExternalToWorld->SetObjectReference(FirstActor);
+			
+			MaterialAndMesh->InstancedMeshComponent->SetStaticMesh(MaterialAndMesh->CylinderMesh);
+			MaterialAndMesh->InstancedMeshComponent->SetMaterial(0, MaterialAndMesh->GradientLinearMaterial);
+		})
+		.ApplySnapshot()
+
+		.RunTest([&]()
+		{
+			TestTrue("Cyclic Reference Restored", FirstActor->HasObjectReference(SecondActor));
+			TestTrue("Reference Not Lost", SecondActor->HasObjectReference(FirstActor));
+			TestTrue("Restore Pointing To Self", PointToSelfActor->HasObjectReference(PointToSelfActor));
+			TestTrue("World > External > World", FromWorldToExternal->HasObjectReference(FirstActor));
+			TestTrue("External > World > External", FromExternalToWorld->HasObjectReference(FromExternalToWorld->CubeMesh));
+
+			TestEqual("Mesh", MaterialAndMesh->InstancedMeshComponent->GetStaticMesh(), MaterialAndMesh->CubeMesh);
+			TestEqual("Material", MaterialAndMesh->InstancedMeshComponent->GetMaterial(0), MaterialAndMesh->GradientLinearMaterial);
+			
+		});
+
+	return true;
+}
+
+IMPLEMENT_SIMPLE_AUTOMATION_TEST(FInstancedStaticMesh, "VirtualProduction.LevelSnapshots.Snapshot.InstancedStaticMesh", (EAutomationTestFlags::ApplicationContextMask | EAutomationTestFlags::EngineFilter));
+bool FInstancedStaticMesh::RunTest(const FString& Parameters)
+{
+	ASnapshotTestActor* MaterialAndMesh = nullptr;
+
+	const FTransform StartFirstTransform(FRotator(100.f).Quaternion(), FVector(110.f) , FVector(1.f));
+	const FTransform EndFirstTransform(FRotator(200.f).Quaternion(), FVector(210.f) , FVector(2.f));
+	const FTransform StartSecondTransform(FRotator(100.f).Quaternion(), FVector(-1100.f) , FVector(-10.f));
+	const FTransform NewThirdTransform(FRotator(-300.f).Quaternion(), FVector(-310.f) , FVector(30.f));
+	
+	FSnapshotTestRunner()
+		.ModifyWorld([&](UWorld* World)
+		{
+			MaterialAndMesh = World->SpawnActor<ASnapshotTestActor>();
+			
+			MaterialAndMesh->InstancedMeshComponent->AddInstance(StartFirstTransform);
+			MaterialAndMesh->InstancedMeshComponent->AddInstance(StartSecondTransform);
+		})
+		.TakeSnapshot()
+
+		.ModifyWorld([&](UWorld* World)
+		{
+			MaterialAndMesh->InstancedMeshComponent->UpdateInstanceTransform(0, EndFirstTransform);
+			MaterialAndMesh->InstancedMeshComponent->AddInstance(NewThirdTransform);
+		})
+		.ApplySnapshot()
+
+		.RunTest([&]()
+		{
+			TestEqual("Instance Count", MaterialAndMesh->InstancedMeshComponent->GetInstanceCount(), 2);
+
+			FTransform ActualFirstTransform, ActualSecondTransform;
+			MaterialAndMesh->InstancedMeshComponent->GetInstanceTransform(0, ActualFirstTransform);
+			MaterialAndMesh->InstancedMeshComponent->GetInstanceTransform(1, ActualSecondTransform);
+
+			// TODO: Investigate rotation not working correcty... it is some kind of math problem that causes angle to flip, i.e. 80 and 100 represent the same angle but are obviously not equal floats...
+			TestTrue("Instance 1", ActualFirstTransform.GetLocation().Equals(StartFirstTransform.GetLocation()) && ActualFirstTransform.GetScale3D().Equals(StartFirstTransform.GetScale3D()));
+			TestTrue("Instance 2", ActualSecondTransform.GetLocation().Equals(StartSecondTransform.GetLocation())&& ActualSecondTransform.GetScale3D().Equals(StartSecondTransform.GetScale3D()));
+		});
+
+	return true;
+}
+
+IMPLEMENT_SIMPLE_AUTOMATION_TEST(FSkipTransientAndDeprecatedProperties, "VirtualProduction.LevelSnapshots.Snapshot.SkipTransientAndDeprecatedProperties", (EAutomationTestFlags::ApplicationContextMask | EAutomationTestFlags::EngineFilter));
+bool FSkipTransientAndDeprecatedProperties::RunTest(const FString& Parameters)
+{
+	ASnapshotTestActor* TestActor = nullptr;
+
+	FSnapshotTestRunner()
+		.ModifyWorld([&](UWorld* World)
+		{
+			TestActor = World->SpawnActor<ASnapshotTestActor>();
+		})
+		.TakeSnapshot()
+
+		.ModifyWorld([&](UWorld* World)
+		{
+			TestActor->DeprecatedProperty_DEPRECATED = 500;
+			TestActor->TransientProperty = 500;
+		})
+		.ApplySnapshot()
+
+		.RunTest([&]()
+		{
+			TestEqual("Skip Deprecated Property", TestActor->DeprecatedProperty_DEPRECATED, 500);
+			TestEqual("Skip Transient Property", TestActor->TransientProperty, 500);
+		});
+
+	return true;
+}
+
diff --git a/Engine/Plugins/VirtualProduction/LevelSnapshots/Source/LevelSnapshots/Private/Tests/SnapshotTestActor.cpp b/Engine/Plugins/VirtualProduction/LevelSnapshots/Source/LevelSnapshots/Private/Tests/SnapshotTestActor.cpp
new file mode 100644
--- /dev/null
+++ b/Engine/Plugins/VirtualProduction/LevelSnapshots/Source/LevelSnapshots/Private/Tests/SnapshotTestActor.cpp
@@ -0,0 +1,110 @@
+﻿// Copyright Epic Games, Inc. All Rights Reserved.
+
+#include "SnapshotTestActor.h"
+
+#include "UObject/ConstructorHelpers.h"
+
+ASnapshotTestActor::ASnapshotTestActor()
+{
+	ConstructorHelpers::FObjectFinder<UMaterialInterface> GradientLinearMaterialFinder(TEXT("Material'/Engine/MaterialTemplates/Gradients/Gradient_Linear.Gradient_Linear'"), LOAD_Quiet | LOAD_NoWarn);
+	if (GradientLinearMaterialFinder.Succeeded())
+	{
+		GradientLinearMaterial = GradientLinearMaterialFinder.Object;
+	}
+
+	ConstructorHelpers::FObjectFinder<UMaterialInterface> GradientRadialMaterialFinder(TEXT("Material'/Engine/MaterialTemplates/Gradients/Gradient_Radial.Gradient_Radial'"), LOAD_Quiet | LOAD_NoWarn);
+	if (GradientRadialMaterialFinder.Succeeded())
+	{
+		GradientRadialMaterial = GradientRadialMaterialFinder.Object;
+	}
+	
+	ConstructorHelpers::FObjectFinder<UStaticMesh> CubeFinder(TEXT("StaticMesh'/Engine/BasicShapes/Cube.Cube'"), LOAD_Quiet | LOAD_NoWarn);
+	if (CubeFinder.Succeeded())
+	{
+		CubeMesh = CubeFinder.Object;
+	}
+	ConstructorHelpers::FObjectFinder<UStaticMesh> CylinderFinder(TEXT("StaticMesh'/Engine/BasicShapes/Cylinder.Cylinder'"), LOAD_Quiet | LOAD_NoWarn);
+	if (CylinderFinder.Succeeded())
+	{
+		CylinderMesh = CylinderFinder.Object;
+	}
+	
+	InstancedMeshComponent = CreateDefaultSubobject<UInstancedStaticMeshComponent>(TEXT("InstancedMeshComponent"));
+	InstancedMeshComponent->SetupAttachment(RootComponent);
+
+	PointLightComponent = CreateDefaultSubobject<UPointLightComponent>(TEXT("PointLightComponent"));
+	PointLightComponent->SetupAttachment(InstancedMeshComponent);
+}
+
+bool ASnapshotTestActor::HasObjectReference(UObject* Object, bool bOnlyCheckCollections, FName MapKey) const
+{
+	const bool bReferencesEqual = (bOnlyCheckCollections
+			|| (ObjectReference == Object && SoftPath == Object && SoftObjectPtr == Object && WeakObjectPtr == Object));
+
+	const bool bCollectionsEqual = ObjectArray.Contains(Object)
+		&& ObjectSet.Contains(Object)
+		&& SoftPathArray.Contains(Object)
+		&& SoftPathSet.Contains(Object)
+		&& SoftObjectPtrArray.Contains(Object)
+		&& SoftObjectPtrSet.Contains(Object)
+		&& WeakObjectPtrArray.Contains(Object)
+		&& WeakObjectPtrSet.Contains(Object)
+		&& ObjectMap.Find(MapKey) && *ObjectMap.Find(MapKey) == Object
+		&& SoftPathMap.Find(MapKey) && *SoftPathMap.Find(MapKey) == Object
+		&& SoftObjectPtrMap.Find(MapKey) && *SoftObjectPtrMap.Find(MapKey) == Object
+		&& WeakObjectPtrMap.Find(MapKey) && *WeakObjectPtrMap.Find(MapKey) == Object;
+
+	return bReferencesEqual && bCollectionsEqual;
+}
+
+void ASnapshotTestActor::SetObjectReference(UObject* Object, FName MapKey)
+{
+	ClearObjectReferences();
+	AddObjectReference(Object, MapKey);
+}
+
+void ASnapshotTestActor::AddObjectReference(UObject* Object, FName MapKey)
+{
+	ObjectReference = Object;
+	ObjectArray.Add(Object);
+	ObjectSet.Add(Object);
+	ObjectMap.Add(MapKey, Object);
+
+	SoftPath = Object;
+	SoftPathArray.Add(Object);
+	SoftPathSet.Add(Object);
+	SoftPathMap.Add(MapKey, Object);
+
+	SoftObjectPtr = Object;
+	SoftObjectPtrArray.Add(Object);
+	SoftObjectPtrSet.Add(Object);
+	SoftObjectPtrMap.Add(MapKey, Object);
+
+	WeakObjectPtr = Object;
+	WeakObjectPtrArray.Add(Object);
+	WeakObjectPtrSet.Add(Object);
+	WeakObjectPtrMap.Add(MapKey, Object);
+}
+
+void ASnapshotTestActor::ClearObjectReferences()
+{
+	ObjectReference = nullptr;
+	ObjectArray.Reset();
+	ObjectSet.Reset();
+	ObjectMap.Reset();
+
+	SoftPath = nullptr;
+	SoftPathArray.Reset();
+	SoftPathSet.Reset();
+	SoftPathMap.Reset();
+
+	SoftObjectPtr = nullptr;
+	SoftObjectPtrArray.Reset();
+	SoftObjectPtrSet.Reset();
+	SoftObjectPtrMap.Reset();
+
+	WeakObjectPtr = nullptr;
+	WeakObjectPtrArray.Reset();
+	WeakObjectPtrSet.Reset();
+	WeakObjectPtrMap.Reset();
+}
diff --git a/Engine/Plugins/VirtualProduction/LevelSnapshots/Source/LevelSnapshots/Private/Tests/SnapshotTestActor.h b/Engine/Plugins/VirtualProduction/LevelSnapshots/Source/LevelSnapshots/Private/Tests/SnapshotTestActor.h
new file mode 100644
--- /dev/null
+++ b/Engine/Plugins/VirtualProduction/LevelSnapshots/Source/LevelSnapshots/Private/Tests/SnapshotTestActor.h
@@ -0,0 +1,131 @@
+﻿// Copyright Epic Games, Inc. All Rights Reserved.
+
+#pragma once
+
+#include "CoreMinimal.h"
+#include "Components/InstancedStaticMeshComponent.h"
+#include "Components/PointLightComponent.h"
+#include "GameFramework/Character.h"
+#include "Materials/MaterialInterface.h"
+#include "SnapshotTestActor.generated.h"
+
+UCLASS()
+class ASnapshotTestActor : public ACharacter
+{
+	GENERATED_BODY()
+public:
+
+	ASnapshotTestActor();
+
+	bool HasObjectReference(UObject* Object, bool bOnlyCheckCollections = false, FName MapKey = NAME_Name) const;
+	
+	void SetObjectReference(UObject* Object, FName MapKey = NAME_Name);
+	void AddObjectReference(UObject* Object, FName MapKey = NAME_Name);
+	void ClearObjectReferences();
+
+
+	UPROPERTY()
+	int32 DeprecatedProperty_DEPRECATED = 100;
+
+	UPROPERTY(EditAnywhere, Transient, Category = "Level Snapshots")
+	int32 TransientProperty = 200;
+
+
+	
+	
+	UPROPERTY(EditAnywhere, Category = "Level Snapshots")
+	UObject* ObjectReference;
+
+	UPROPERTY(EditAnywhere, Category = "Level Snapshots")
+	TArray<UObject*> ObjectArray;
+
+	UPROPERTY(EditAnywhere, Category = "Level Snapshots")
+	TSet<UObject*> ObjectSet;
+
+	UPROPERTY(EditAnywhere, Category = "Level Snapshots")
+	TMap<FName, UObject*> ObjectMap;
+
+
+
+	
+
+	UPROPERTY(EditAnywhere, Category = "Level Snapshots")
+	FSoftObjectPath SoftPath;
+
+	UPROPERTY(EditAnywhere, Category = "Level Snapshots")
+	TArray<FSoftObjectPath> SoftPathArray;
+
+	UPROPERTY(EditAnywhere, Category = "Level Snapshots")
+	TSet<FSoftObjectPath> SoftPathSet;
+
+	UPROPERTY(EditAnywhere, Category = "Level Snapshots")
+	TMap<FName, FSoftObjectPath> SoftPathMap;
+
+	
+
+	
+
+	UPROPERTY(EditAnywhere, Category = "Level Snapshots")
+	TSoftObjectPtr<UObject> SoftObjectPtr;
+
+	UPROPERTY(EditAnywhere, Category = "Level Snapshots")
+	TArray<TSoftObjectPtr<UObject>> SoftObjectPtrArray;
+
+	UPROPERTY(EditAnywhere, Category = "Level Snapshots")
+	TSet<TSoftObjectPtr<UObject>> SoftObjectPtrSet;
+
+	UPROPERTY(EditAnywhere, Category = "Level Snapshots")
+	TMap<FName, TSoftObjectPtr<UObject>> SoftObjectPtrMap;
+
+
+
+	
+
+	UPROPERTY(EditAnywhere, Category = "Level Snapshots")
+	TWeakObjectPtr<UObject> WeakObjectPtr;
+
+	UPROPERTY(EditAnywhere, Category = "Level Snapshots")
+	TArray<TWeakObjectPtr<UObject>> WeakObjectPtrArray;
+
+	UPROPERTY(EditAnywhere, Category = "Level Snapshots")
+	TSet<TWeakObjectPtr<UObject>> WeakObjectPtrSet;
+
+	UPROPERTY(EditAnywhere, Category = "Level Snapshots")
+	TMap<FName, TWeakObjectPtr<UObject>> WeakObjectPtrMap;
+
+	
+
+	
+	
+	UPROPERTY(EditAnywhere, Category = "Level Snapshots")
+	UActorComponent* ExternalComponentReference;
+
+	UPROPERTY(EditAnywhere, Category = "Level Snapshots")
+	UObject* ExternalComponentReferenceAsUObject;
+
+
+
+	
+	
+	UPROPERTY(EditAnywhere, Category = "Level Snapshots")
+	UMaterialInterface* GradientLinearMaterial;
+
+	UPROPERTY(EditAnywhere, Category = "Level Snapshots")
+	UMaterialInterface* GradientRadialMaterial;
+	
+	UPROPERTY(EditAnywhere, Category = "Level Snapshots")
+	UStaticMesh* CubeMesh;
+
+	UPROPERTY(EditAnywhere, Category = "Level Snapshots")
+	UStaticMesh* CylinderMesh;
+
+
+
+	
+
+	UPROPERTY(EditAnywhere, Category = "Level Snapshots")
+	UInstancedStaticMeshComponent* InstancedMeshComponent;
+
+	UPROPERTY(EditAnywhere, Category = "Level Snapshots")
+	UPointLightComponent* PointLightComponent;
+};
diff --git a/Engine/Plugins/VirtualProduction/LevelSnapshots/Source/LevelSnapshots/Private/Tests/SnapshotTestRunner.cpp b/Engine/Plugins/VirtualProduction/LevelSnapshots/Source/LevelSnapshots/Private/Tests/SnapshotTestRunner.cpp
new file mode 100644
--- /dev/null
+++ b/Engine/Plugins/VirtualProduction/LevelSnapshots/Source/LevelSnapshots/Private/Tests/SnapshotTestRunner.cpp
@@ -0,0 +1,96 @@
+﻿// Copyright Epic Games, Inc. All Rights Reserved.
+
+#include "SnapshotTestRunner.h"
+
+#include "LevelSnapshotsFunctionLibrary.h"
+
+FName FSnapshotTestRunner::DefaultSnapshotId = FName("DefaultSnapshotId");
+
+FSnapshotTestRunner::FSnapshotTestRunner()
+{
+	TestWorld = MakeShared<FPreviewScene>(
+		FPreviewScene::ConstructionValues()
+			.SetEditor(true)
+			);
+}
+
+FSnapshotTestRunner& FSnapshotTestRunner::ModifyWorld(TFunction<void(UWorld*)> Callback)
+{
+	Callback(TestWorld->GetWorld());
+	return *this;
+}
+
+FSnapshotTestRunner& FSnapshotTestRunner::TakeSnapshot(FName SnapshotId)
+{
+	if (ULevelSnapshot** ExistingSnapshot = Snapshots.Find(SnapshotId))
+	{
+		(*ExistingSnapshot)->SnapshotWorld(TestWorld->GetWorld());
+	}
+	else
+	{
+		Snapshots.Add(
+			SnapshotId,
+			ULevelSnapshotsFunctionLibrary::TakeLevelSnapshot(TestWorld->GetWorld(), SnapshotId)
+			);
+	}
+	
+	return *this;
+}
+
+FSnapshotTestRunner& FSnapshotTestRunner::AccessSnapshot(TFunction<void (ULevelSnapshot*)> Callback, FName SnapshotId)
+{
+	if (ULevelSnapshot** ExistingSnapshot = Snapshots.Find(SnapshotId))
+	{
+		Callback(*ExistingSnapshot);
+	}
+	else
+	{
+		checkNoEntry();
+	}
+
+	return *this;
+}
+
+FSnapshotTestRunner& FSnapshotTestRunner::ApplySnapshot(TFunction<ULevelSnapshotFilter*()> Callback, FName SnapshotId)
+{
+	return ApplySnapshot(Callback(), SnapshotId);
+}
+
+FSnapshotTestRunner& FSnapshotTestRunner::ApplySnapshot(ULevelSnapshotFilter* Filter, FName SnapshotId)
+{
+	if (ULevelSnapshot** ExistingSnapshot = Snapshots.Find(SnapshotId))
+	{
+		ULevelSnapshotsFunctionLibrary::ApplySnapshotToWorld(TestWorld->GetWorld(), *ExistingSnapshot, Filter);
+	}
+	else
+	{
+		checkNoEntry();
+	}
+	
+	return *this;
+}
+
+FSnapshotTestRunner& FSnapshotTestRunner::ApplySnapshot(TFunction<FPropertySelectionMap()> Callback, FName SnapshotId)
+{
+	return ApplySnapshot(Callback(), SnapshotId);
+}
+
+FSnapshotTestRunner& FSnapshotTestRunner::ApplySnapshot(const FPropertySelectionMap& SelectionSet, FName SnapshotId)
+{
+	if (ULevelSnapshot** ExistingSnapshot = Snapshots.Find(SnapshotId))
+	{
+		(*ExistingSnapshot)->ApplySnapshotToWorld(TestWorld->GetWorld(), SelectionSet);
+	}
+	else
+	{
+		checkNoEntry();
+	}
+	
+	return *this;
+}
+
+FSnapshotTestRunner& FSnapshotTestRunner::RunTest(TFunction<void()> Callback)
+{
+	Callback();
+	return *this;
+}
diff --git a/Engine/Plugins/VirtualProduction/LevelSnapshots/Source/LevelSnapshots/Private/Tests/SnapshotTestRunner.h b/Engine/Plugins/VirtualProduction/LevelSnapshots/Source/LevelSnapshots/Private/Tests/SnapshotTestRunner.h
new file mode 100644
--- /dev/null
+++ b/Engine/Plugins/VirtualProduction/LevelSnapshots/Source/LevelSnapshots/Private/Tests/SnapshotTestRunner.h
@@ -0,0 +1,44 @@
+﻿// Copyright Epic Games, Inc. All Rights Reserved.
+
+#pragma once
+
+#include "CoreMinimal.h"
+
+#include "LevelSnapshot.h"
+#include "LevelSnapshotFilters.h"
+
+#include "Engine/World.h"
+
+class FPreviewScene;
+
+/**
+ * Utility for executing tests
+ */
+class FSnapshotTestRunner
+{
+public:
+
+	static FName DefaultSnapshotId;
+
+	FSnapshotTestRunner();
+
+	FSnapshotTestRunner& ModifyWorld(TFunction<void(UWorld*)> Callback);
+	
+	FSnapshotTestRunner& TakeSnapshot(FName SnapshotId = DefaultSnapshotId);
+	FSnapshotTestRunner& AccessSnapshot(TFunction<void(ULevelSnapshot*)> Callback, FName SnapshotId = DefaultSnapshotId);
+	
+	FSnapshotTestRunner& ApplySnapshot(TFunction<ULevelSnapshotFilter*()> Callback, FName SnapshotId = DefaultSnapshotId);
+	FSnapshotTestRunner& ApplySnapshot(ULevelSnapshotFilter* Filter = nullptr, FName SnapshotId = DefaultSnapshotId);
+	
+	FSnapshotTestRunner& ApplySnapshot(TFunction<FPropertySelectionMap()> Callback, FName SnapshotId = DefaultSnapshotId);
+	FSnapshotTestRunner& ApplySnapshot(const FPropertySelectionMap& SelectionSet, FName SnapshotId = DefaultSnapshotId);
+
+	/* Just calls Callback. Existing for better readability in tests.  */
+	FSnapshotTestRunner& RunTest(TFunction<void()> Callback);
+
+private:
+
+	TSharedPtr<FPreviewScene> TestWorld;
+
+	TMap<FName, ULevelSnapshot*> Snapshots;
+};
