Fix unit test

diff --git a/Engine/Plugins/VirtualProduction/LevelSnapshots/Source/LevelSnapshots/Private/Tests/ApplyToWorldTests.cpp b/Engine/Plugins/VirtualProduction/LevelSnapshots/Source/LevelSnapshots/Private/Tests/ApplyToWorldTests.cpp
--- a/Engine/Plugins/VirtualProduction/LevelSnapshots/Source/LevelSnapshots/Private/Tests/ApplyToWorldTests.cpp
+++ b/Engine/Plugins/VirtualProduction/LevelSnapshots/Source/LevelSnapshots/Private/Tests/ApplyToWorldTests.cpp
@@ -184,7 +184,7 @@ bool FInstancedStaticMesh::RunTest(const FString& Parameters)
 
 	const FTransform StartFirstTransform(FRotator(100.f).Quaternion(), FVector(110.f) , FVector(1.f));
 	const FTransform EndFirstTransform(FRotator(200.f).Quaternion(), FVector(210.f) , FVector(2.f));
-	const FTransform StartSecondTransform(FRotator(100.f).Quaternion(), FVector(-1100.f) , FVector(-10.f));
+	const FTransform StartSecondTransform(FRotator(100.f).Quaternion(), FVector(-1100.f) , FVector(1.f));
 	const FTransform NewThirdTransform(FRotator(-300.f).Quaternion(), FVector(-310.f) , FVector(30.f));
 	
 	FSnapshotTestRunner()
