Fix compile error

diff --git a/Engine/Plugins/VirtualProduction/LevelSnapshots/Source/LevelSnapshots/Private/Tests/SnapshotTestRunner.cpp b/Engine/Plugins/VirtualProduction/LevelSnapshots/Source/LevelSnapshots/Private/Tests/SnapshotTestRunner.cpp
--- a/Engine/Plugins/VirtualProduction/LevelSnapshots/Source/LevelSnapshots/Private/Tests/SnapshotTestRunner.cpp
+++ b/Engine/Plugins/VirtualProduction/LevelSnapshots/Source/LevelSnapshots/Private/Tests/SnapshotTestRunner.cpp
@@ -3,6 +3,7 @@
 #include "SnapshotTestRunner.h"
 
 #include "LevelSnapshotsFunctionLibrary.h"
+#include "PreviewScene.h"
 
 FName FSnapshotTestRunner::DefaultSnapshotId = FName("DefaultSnapshotId");
 
