Remove negation for all filters and filter rows.

diff --git a/Engine/Plugins/VirtualProduction/LevelSnapshots/Source/LevelSnapshotsEditor/Private/Data/Filters/ConjunctionFilter.cpp b/Engine/Plugins/VirtualProduction/LevelSnapshots/Source/LevelSnapshotsEditor/Private/Data/Filters/ConjunctionFilter.cpp
--- a/Engine/Plugins/VirtualProduction/LevelSnapshots/Source/LevelSnapshotsEditor/Private/Data/Filters/ConjunctionFilter.cpp
+++ b/Engine/Plugins/VirtualProduction/LevelSnapshots/Source/LevelSnapshotsEditor/Private/Data/Filters/ConjunctionFilter.cpp
@@ -30,18 +30,14 @@ namespace
 		return bAtLeastOneFilterSaidInclude ? EFilterResult::Include : EFilterResult::DoNotCare;
 	}
 	
-	EFilterResult::Type ExecuteAndChain(const TArray<UNegatableFilter*>& Children, EEditorFilterBehavior EditorFilterBehavior, ConjunctionFilterCallback FilterCallback)
+	EFilterResult::Type ExecuteAndChain(const TArray<UNegatableFilter*>& Children, ConjunctionFilterCallback FilterCallback)
 	{
-		if (Children.Num() == 0 
-            || EditorFilterBehavior == EEditorFilterBehavior::Ignore
-            || !ensure(EditorFilterBehavior != EEditorFilterBehavior::Mixed))
+		if (Children.Num() == 0 )
 		{
 			return EFilterResult::DoNotCare;
 		}
 		
-		const EFilterResult::Type IntermediateResult = AndChain(Children, FilterCallback);
-		const bool bShouldNegate = EditorFilterBehavior == EEditorFilterBehavior::Negate;
-		return bShouldNegate ? EFilterResult::Negate(IntermediateResult) : IntermediateResult;
+		return  AndChain(Children, FilterCallback);
 	}
 }
 
@@ -55,9 +51,6 @@ UNegatableFilter* UConjunctionFilter::CreateChild(const TSubclassOf<ULevelSnapsh
 	ULevelSnapshotFilter* FilterImplementation = NewObject<ULevelSnapshotFilter>(this, FilterClass.Get());
 	UNegatableFilter* Child = UNegatableFilter::CreateNegatableFilter(FilterImplementation, this);
 
-	// Set Child parent
-	Child->SetParentFilter(this);
-
 	Children.Add(Child);
 	OnChildAdded.Broadcast(Child);
 	
@@ -81,7 +74,7 @@ const TArray<UNegatableFilter*>& UConjunctionFilter::GetChildren() const
 
 EFilterResult::Type UConjunctionFilter::IsActorValid(const FIsActorValidParams& Params) const
 {
-	return ExecuteAndChain(Children, EditorFilterBehavior, [&Params](UNegatableFilter* Child)
+	return ExecuteAndChain(Children, [&Params](UNegatableFilter* Child)
 		{
 			return Child->IsActorValid(Params);
 		});
@@ -89,7 +82,7 @@ EFilterResult::Type UConjunctionFilter::IsActorValid(const FIsActorValidParams&
 
 EFilterResult::Type UConjunctionFilter::IsPropertyValid(const FIsPropertyValidParams& Params) const
 {
-	return ExecuteAndChain(Children, EditorFilterBehavior, [&Params](UNegatableFilter* Child)
+	return ExecuteAndChain(Children, [&Params](UNegatableFilter* Child)
 		{
 			return Child->IsPropertyValid(Params);
 		});
@@ -97,7 +90,7 @@ EFilterResult::Type UConjunctionFilter::IsPropertyValid(const FIsPropertyValidPa
 
 EFilterResult::Type UConjunctionFilter::IsDeletedActorValid(const FIsDeletedActorValidParams& Params) const
 {
-	return ExecuteAndChain(Children, EditorFilterBehavior, [&Params](UNegatableFilter* Child)
+	return ExecuteAndChain(Children, [&Params](UNegatableFilter* Child)
 		{
 			return Child->IsDeletedActorValid(Params);
 		});
@@ -105,51 +98,8 @@ EFilterResult::Type UConjunctionFilter::IsDeletedActorValid(const FIsDeletedActo
 
 EFilterResult::Type UConjunctionFilter::IsAddedActorValid(const FIsAddedActorValidParams& Params) const
 {
-	return ExecuteAndChain(Children, EditorFilterBehavior, [&Params](UNegatableFilter* Child)
+	return ExecuteAndChain(Children, [&Params](UNegatableFilter* Child)
 		{
 			return Child->IsAddedActorValid(Params);
 		});
 }
-
-TArray<UEditorFilter*> UConjunctionFilter::GetEditorChildren()
-{
-	TArray<UEditorFilter*> EditorFilterChildren;
-
-	for (UNegatableFilter* ChildFilter : Children)
-	{
-		EditorFilterChildren.Add(ChildFilter);
-	}
-
-	return EditorFilterChildren;
-}
-
-void UConjunctionFilter::IncrementEditorFilterBehavior(const bool bIncludeChildren)
-{
-	switch (EditorFilterBehavior)
-	{
-	case EEditorFilterBehavior::DoNotNegate:
-		EditorFilterBehavior = EEditorFilterBehavior::Negate;
-		break;
-	case EEditorFilterBehavior::Negate:
-		EditorFilterBehavior = EEditorFilterBehavior::Ignore;
-		break;
-	case EEditorFilterBehavior::Ignore:
-		EditorFilterBehavior = EEditorFilterBehavior::DoNotNegate;
-		break;
-	default:
-		checkNoEntry();
-	}
-
-	if (bIncludeChildren)
-	{
-		UpdateAllChildrenEditorFilterBehavior(EditorFilterBehavior, bIncludeChildren);
-	}
-}
-
-void UConjunctionFilter::SetEditorFilterBehavior(const EEditorFilterBehavior InFilterBehavior, const bool bIncludeChildren)
-{
-	if (ensureMsgf(EditorFilterBehavior != EEditorFilterBehavior::Mixed, TEXT("Internal error. Conjunction filter cannot have mixed behaviour.")))
-	{
-		EditorFilterBehavior = InFilterBehavior;
-	}
-}
diff --git a/Engine/Plugins/VirtualProduction/LevelSnapshots/Source/LevelSnapshotsEditor/Private/Data/Filters/ConjunctionFilter.h b/Engine/Plugins/VirtualProduction/LevelSnapshots/Source/LevelSnapshotsEditor/Private/Data/Filters/ConjunctionFilter.h
--- a/Engine/Plugins/VirtualProduction/LevelSnapshots/Source/LevelSnapshotsEditor/Private/Data/Filters/ConjunctionFilter.h
+++ b/Engine/Plugins/VirtualProduction/LevelSnapshots/Source/LevelSnapshotsEditor/Private/Data/Filters/ConjunctionFilter.h
@@ -3,8 +3,7 @@
 #pragma once
 
 #include "CoreMinimal.h"
-
-#include "EditorFilter.h"
+#include "LevelSnapshotFilters.h"
 #include "ConjunctionFilter.generated.h"
 
 class UNegatableFilter;
@@ -14,7 +13,7 @@ class UNegatableFilter;
  * It is valid to have no children: in this case, this filter return false.
  */
 UCLASS(meta = (InternalSnapshotFilter))
-class UConjunctionFilter : public UEditorFilter
+class UConjunctionFilter : public ULevelSnapshotFilter
 {
 	GENERATED_BODY()
 public:
@@ -33,12 +32,6 @@ public:
 	virtual EFilterResult::Type IsDeletedActorValid(const FIsDeletedActorValidParams& Params) const override;
 	virtual EFilterResult::Type IsAddedActorValid(const FIsAddedActorValidParams& Params) const override;
 	//~ End ULevelSnapshotFilter Interface
-	
-	//~ Begin UEditorFilter Interface
-	virtual TArray<UEditorFilter*> GetEditorChildren();
-	virtual void IncrementEditorFilterBehavior(const bool bIncludeChildren = false) override;
-	virtual void SetEditorFilterBehavior(const EEditorFilterBehavior InFilterBehavior, const bool bIncludeChildren = false);
-	//~ Begin UEditorFilter Interface
 
 	DECLARE_EVENT_OneParam(UConjunctionFilter, FOnChildModified, UNegatableFilter*);
 	FOnChildModified OnChildAdded;
diff --git a/Engine/Plugins/VirtualProduction/LevelSnapshots/Source/LevelSnapshotsEditor/Private/Data/Filters/DisjunctiveNormalFormFilter.cpp b/Engine/Plugins/VirtualProduction/LevelSnapshots/Source/LevelSnapshotsEditor/Private/Data/Filters/DisjunctiveNormalFormFilter.cpp
--- a/Engine/Plugins/VirtualProduction/LevelSnapshots/Source/LevelSnapshotsEditor/Private/Data/Filters/DisjunctiveNormalFormFilter.cpp
+++ b/Engine/Plugins/VirtualProduction/LevelSnapshots/Source/LevelSnapshotsEditor/Private/Data/Filters/DisjunctiveNormalFormFilter.cpp
@@ -58,15 +58,6 @@ namespace
 UConjunctionFilter* UDisjunctiveNormalFormFilter::CreateChild()
 {
 	UConjunctionFilter* Child = NewObject<UConjunctionFilter>(this);
-
-	// Set Child parent
-	Child->SetParentFilter(this);
-
-	if (EditorFilterBehavior != EEditorFilterBehavior::Mixed)
-	{
-		Child->SetEditorFilterBehavior(EditorFilterBehavior, false); // Set Filter Result of parent
-	}
-
 	Children.Add(Child);
 	return Child;
 }
@@ -113,15 +104,3 @@ EFilterResult::Type UDisjunctiveNormalFormFilter::IsAddedActorValid(const FIsAdd
 		return Child->IsAddedActorValid(Params);
 	});
 }
-
-TArray<UEditorFilter*> UDisjunctiveNormalFormFilter::GetEditorChildren()
-{
-	TArray<UEditorFilter*> EditorFilterChildren;
-
-	for (UConjunctionFilter* ChildFilter : Children)
-	{
-		EditorFilterChildren.Add(ChildFilter);
-	}
-
-	return EditorFilterChildren;
-}
diff --git a/Engine/Plugins/VirtualProduction/LevelSnapshots/Source/LevelSnapshotsEditor/Private/Data/Filters/DisjunctiveNormalFormFilter.h b/Engine/Plugins/VirtualProduction/LevelSnapshots/Source/LevelSnapshotsEditor/Private/Data/Filters/DisjunctiveNormalFormFilter.h
--- a/Engine/Plugins/VirtualProduction/LevelSnapshots/Source/LevelSnapshotsEditor/Private/Data/Filters/DisjunctiveNormalFormFilter.h
+++ b/Engine/Plugins/VirtualProduction/LevelSnapshots/Source/LevelSnapshotsEditor/Private/Data/Filters/DisjunctiveNormalFormFilter.h
@@ -3,8 +3,7 @@
 #pragma once
 
 #include "CoreMinimal.h"
-
-#include "EditorFilter.h"
+#include "LevelSnapshotFilters.h"
 #include "DisjunctiveNormalFormFilter.generated.h"
 
 class UConjunctionFilter;
@@ -16,7 +15,7 @@ class UConjunctionFilter;
  * Disjunctive normal form = ORs of ANDs. Example: (a && !b) || (c && d) || e
  */
 UCLASS(meta = (InternalSnapshotFilter))
-class UDisjunctiveNormalFormFilter : public UEditorFilter
+class UDisjunctiveNormalFormFilter : public ULevelSnapshotFilter
 {
 	GENERATED_BODY()
 public:
@@ -34,14 +33,6 @@ public:
 	virtual EFilterResult::Type IsAddedActorValid(const FIsAddedActorValidParams& Params) const override;
 	//~ End ULevelSnapshotFilter Interface
 	
-	//~ Begin UEditorFilter Interface
-	virtual TArray<UEditorFilter*> GetEditorChildren();
-	//~ Begin UEditorFilter Interface
-
-	DECLARE_EVENT_TwoParams(UDisjunctiveNormalFormFilter, FOnChildrenChanged, UConjunctionFilter* /*AffectedFilter*/, int32 /*ChildIndex*/);
-	FOnChildrenChanged OnChildAdded;
-	FOnChildrenChanged OnChildRemoved;
-	
 private:
 
 	UPROPERTY()
diff --git a/Engine/Plugins/VirtualProduction/LevelSnapshots/Source/LevelSnapshotsEditor/Private/Data/Filters/EditorFilter.cpp b/Engine/Plugins/VirtualProduction/LevelSnapshots/Source/LevelSnapshotsEditor/Private/Data/Filters/EditorFilter.cpp
deleted file mode 100644
diff --git a/Engine/Plugins/VirtualProduction/LevelSnapshots/Source/LevelSnapshotsEditor/Private/Data/Filters/EditorFilter.h b/Engine/Plugins/VirtualProduction/LevelSnapshots/Source/LevelSnapshotsEditor/Private/Data/Filters/EditorFilter.h
deleted file mode 100644
diff --git a/Engine/Plugins/VirtualProduction/LevelSnapshots/Source/LevelSnapshotsEditor/Private/Data/Filters/NegatableFilter.cpp b/Engine/Plugins/VirtualProduction/LevelSnapshots/Source/LevelSnapshotsEditor/Private/Data/Filters/NegatableFilter.cpp
--- a/Engine/Plugins/VirtualProduction/LevelSnapshots/Source/LevelSnapshotsEditor/Private/Data/Filters/NegatableFilter.cpp
+++ b/Engine/Plugins/VirtualProduction/LevelSnapshots/Source/LevelSnapshotsEditor/Private/Data/Filters/NegatableFilter.cpp
@@ -4,15 +4,15 @@
 
 namespace
 {
-	EFilterResult::Type ApplyFilter(EEditorFilterBehavior NegationBehaviour, TFunction<EFilterResult::Type()> ApplyFilter)
+	EFilterResult::Type ApplyFilter(EFilterBehavior NegationBehaviour, TFunction<EFilterResult::Type()> ApplyFilter)
 	{
-		if (NegationBehaviour == EEditorFilterBehavior::Ignore)
+		if (NegationBehaviour == EFilterBehavior::Ignore)
 		{
 			return EFilterResult::DoNotCare; 
 		}
 
 		const EFilterResult::Type Result = ApplyFilter();
-		const bool bShouldNegate = NegationBehaviour == EEditorFilterBehavior::Negate;
+		const bool bShouldNegate = NegationBehaviour == EFilterBehavior::Negate;
 		return bShouldNegate ? EFilterResult::Negate(Result) : Result;
 	}
 }
@@ -30,9 +30,28 @@ UNegatableFilter* UNegatableFilter::CreateNegatableFilter(ULevelSnapshotFilter*
 	return Result;
 }
 
-ULevelSnapshotFilter* UNegatableFilter::GetChildFilter() const
+void UNegatableFilter::IncrementFilterBehaviour()
 {
-	return ChildFilter;
+	switch (FilterBehavior)
+	{
+	case EFilterBehavior::DoNotNegate:
+		SetFilterBehaviour(EFilterBehavior::Negate);
+		break;
+	case EFilterBehavior::Negate:
+		SetFilterBehaviour(EFilterBehavior::Ignore);
+		break;
+	case EFilterBehavior::Ignore:
+		SetFilterBehaviour(EFilterBehavior::DoNotNegate);
+		break;
+	default:
+		// Needed for data migration
+		SetFilterBehaviour(EFilterBehavior::DoNotNegate);
+	}
+}
+
+void UNegatableFilter::SetFilterBehaviour(EFilterBehavior NewFilterBehavior)
+{
+	FilterBehavior = NewFilterBehavior;
 }
 
 FText UNegatableFilter::GetDisplayName() const
@@ -49,7 +68,7 @@ FText UNegatableFilter::GetDisplayName() const
 
 EFilterResult::Type UNegatableFilter::IsActorValid(const FIsActorValidParams& Params) const
 {
-	return ApplyFilter(EditorFilterBehavior, [this, &Params]()
+	return ApplyFilter(FilterBehavior, [this, &Params]()
 	{
 		// Child filter is nullptr when user deletes filter class and force deletes
 		return ChildFilter ? ChildFilter->IsActorValid(Params) : EFilterResult::DoNotCare;
@@ -58,7 +77,7 @@ EFilterResult::Type UNegatableFilter::IsActorValid(const FIsActorValidParams& Pa
 
 EFilterResult::Type UNegatableFilter::IsPropertyValid(const FIsPropertyValidParams& Params) const
 {
-	return ApplyFilter(EditorFilterBehavior, [this, &Params]()
+	return ApplyFilter(FilterBehavior, [this, &Params]()
 	{
 		// Child filter is nullptr when user deletes filter class and force deletes
 		return ChildFilter ? ChildFilter->IsPropertyValid(Params) : EFilterResult::DoNotCare;
@@ -67,7 +86,7 @@ EFilterResult::Type UNegatableFilter::IsPropertyValid(const FIsPropertyValidPara
 
 EFilterResult::Type UNegatableFilter::IsDeletedActorValid(const FIsDeletedActorValidParams& Params) const
 {
-	return ApplyFilter(EditorFilterBehavior, [this, &Params]()
+	return ApplyFilter(FilterBehavior, [this, &Params]()
 	{
 		// Child filter is nullptr when user deletes filter class and force deletes
 		return ChildFilter ? ChildFilter->IsDeletedActorValid(Params) : EFilterResult::DoNotCare;
@@ -76,7 +95,7 @@ EFilterResult::Type UNegatableFilter::IsDeletedActorValid(const FIsDeletedActorV
 
 EFilterResult::Type UNegatableFilter::IsAddedActorValid(const FIsAddedActorValidParams& Params) const
 {
-	return ApplyFilter(EditorFilterBehavior, [this, &Params]()
+	return ApplyFilter(FilterBehavior, [this, &Params]()
 	{
 		// Child filter is nullptr when user deletes filter class and force deletes
 		return ChildFilter ? ChildFilter->IsAddedActorValid(Params) : EFilterResult::DoNotCare;
diff --git a/Engine/Plugins/VirtualProduction/LevelSnapshots/Source/LevelSnapshotsEditor/Private/Data/Filters/NegatableFilter.h b/Engine/Plugins/VirtualProduction/LevelSnapshots/Source/LevelSnapshotsEditor/Private/Data/Filters/NegatableFilter.h
--- a/Engine/Plugins/VirtualProduction/LevelSnapshots/Source/LevelSnapshotsEditor/Private/Data/Filters/NegatableFilter.h
+++ b/Engine/Plugins/VirtualProduction/LevelSnapshots/Source/LevelSnapshotsEditor/Private/Data/Filters/NegatableFilter.h
@@ -3,14 +3,25 @@
 #pragma once
 
 #include "CoreMinimal.h"
-#include "EditorFilter.h"
+#include "LevelSnapshotFilters.h"
 #include "NegatableFilter.generated.h"
 
+UENUM()
+enum class EFilterBehavior : uint8
+{
+	/* Pass on same result */
+	DoNotNegate,
+	/* Invert the result */ 
+	Negate,
+	/* Ignore the result */
+	Ignore
+};
+
 /*
  * Calls a child filter and possibly negates its results.
  */
 UCLASS(meta = (InternalSnapshotFilter))
-class UNegatableFilter : public UEditorFilter
+class UNegatableFilter : public ULevelSnapshotFilter
 {
 	GENERATED_BODY()
 public:
@@ -18,8 +29,12 @@ public:
 	/* Wraps the given filter with a negation. Defaults to ChildFilter's outer. */
 	static UNegatableFilter* CreateNegatableFilter(ULevelSnapshotFilter* ChildFilter, const TOptional<UObject*>& Outer = TOptional<UObject*>());
 
-	ULevelSnapshotFilter* GetChildFilter() const;
+	void IncrementFilterBehaviour();
+	void SetFilterBehaviour(EFilterBehavior NewFilterBehavior);
+	
 	FText GetDisplayName() const;
+	EFilterBehavior GetFilterBehavior() const { return FilterBehavior; }
+	ULevelSnapshotFilter* GetChildFilter() const { return ChildFilter; }
 
 	//~ Begin ULevelSnapshotFilter Interface
 	virtual EFilterResult::Type IsActorValid(const FIsActorValidParams& Params) const override;
@@ -36,6 +51,10 @@ public: // Only public to use GET_MEMBER_NAME_CHECKED with compiler checks - do
 	FString Name;
 
 private:
+
+	/* Whether to pass on the result of the filter, negate it, or ignore it. */
+	UPROPERTY(EditAnywhere, Category = "Filter")
+	EFilterBehavior FilterBehavior = EFilterBehavior::DoNotNegate;
 	
 	UPROPERTY()
 	ULevelSnapshotFilter* ChildFilter;
diff --git a/Engine/Plugins/VirtualProduction/LevelSnapshots/Source/LevelSnapshotsEditor/Private/Views/Filter/SLevelSnapshotsEditorFilters.cpp b/Engine/Plugins/VirtualProduction/LevelSnapshots/Source/LevelSnapshotsEditor/Private/Views/Filter/SLevelSnapshotsEditorFilters.cpp
--- a/Engine/Plugins/VirtualProduction/LevelSnapshots/Source/LevelSnapshotsEditor/Private/Views/Filter/SLevelSnapshotsEditorFilters.cpp
+++ b/Engine/Plugins/VirtualProduction/LevelSnapshots/Source/LevelSnapshotsEditor/Private/Views/Filter/SLevelSnapshotsEditorFilters.cpp
@@ -1,24 +1,22 @@
 // Copyright Epic Games, Inc. All Rights Reserved.
 
-#include "SLevelSnapshotsEditorFilters.h"
+#include "Views/Filter/SLevelSnapshotsEditorFilters.h"
 
-#include "DisjunctiveNormalFormFilter.h"
-#include "EditorFilter.h"
-#include "EditorFontGlyphs.h"
-#include "FilteredResults.h"
-#include "ILevelSnapshotsEditorView.h"
-#include "LevelSnapshotsEditorData.h"
+#include "Customizations/NegatableFilterDetailsCustomization.h"
+#include "Data/FilteredResults.h"
+#include "Data/Filters/DisjunctiveNormalFormFilter.h"
+#include "Data/LevelSnapshotsEditorData.h"
 #include "LevelSnapshotsEditorFilters.h"
 #include "LevelSnapshotsEditorStyle.h"
-#include "SFavoriteFilterList.h"
-#include "SMasterFilterIndicatorButton.h"
-#include "SLevelSnapshotsEditorFilterRow.h"
-#include "SSaveAndLoadFilters.h"
+#include "TempInterfaces/ILevelSnapshotsEditorView.h"
+#include "Widgets/Filter/SFavoriteFilterList.h"
+#include "Widgets/SLevelSnapshotsEditorFilterRow.h"
+#include "Widgets/Filter/SSaveAndLoadFilters.h"
 
+#include "EditorFontGlyphs.h"
 #include "EditorStyleSet.h"
 #include "Framework/Notifications/NotificationManager.h"
 #include "IDetailsView.h"
-#include "NegatableFilterDetailsCustomization.h"
 #include "Modules/ModuleManager.h"
 #include "PropertyEditorModule.h"
 #include "Widgets/Input/SButton.h"
@@ -125,10 +123,15 @@ void SLevelSnapshotsEditorFilters::Construct(const FArguments& InArgs, const TSh
 
 	// Create a property view
 	FPropertyEditorModule& EditModule = FModuleManager::Get().GetModuleChecked<FPropertyEditorModule>("PropertyEditor");
-
-	FDetailsViewArgs DetailsViewArgs;
-	DetailsViewArgs.NameAreaSettings = FDetailsViewArgs::HideNameArea;
-	DetailsViewArgs.bHideSelectionTip = true;
+	FDetailsViewArgs DetailsViewArgs(
+		/*bUpdateFromSelection=*/ false,
+		/*bLockable=*/ false, 
+		/*bAllowSearch=*/ true,
+		FDetailsViewArgs::HideNameArea,
+		/*bHideSelectionTip=*/ true,
+		/*InNotifyHook=*/ nullptr,
+		/*InSearchInitialKeyFocus=*/ false,
+		/*InViewIdentifier=*/ NAME_None);
 	DetailsViewArgs.DefaultsOnlyVisibility = EEditDefaultsOnlyNodeVisibility::Automatic;
 	
 	FilterDetailsView = EditModule.CreateDetailView(DetailsViewArgs);
@@ -159,61 +162,35 @@ void SLevelSnapshotsEditorFilters::Construct(const FArguments& InArgs, const TSh
 			[
 				SNew(SHorizontalBox)
 
+				// Refresh button
 				+ SHorizontalBox::Slot()
-				.HAlign(HAlign_Left)
-				.Padding(FMargin(0.f, 2.f))
+				.AutoWidth()
+				.Padding(0.f, 0.f, 2.f, 0.f)
 				[
-					SNew(SHorizontalBox)
-
-					+ SHorizontalBox::Slot()
-					.AutoWidth()
+					SNew(SButton)
+					.ButtonStyle(FEditorStyle::Get(), "FlatButton.Success")
+					.ForegroundColor(FSlateColor::UseForeground())
+					.OnClicked(this, &SLevelSnapshotsEditorFilters::OnClickUpdateResultsView)
 					[
-						// Checkbox 
 						SNew(SHorizontalBox)
-						+ SHorizontalBox::Slot()
-						.Padding(2.f, 0.f)
-						.AutoWidth()
+						+SHorizontalBox::Slot()
+							.AutoWidth()
+							.VAlign(VAlign_Center)
 						[
-							SAssignNew(MasterFilterIndicatorButton, SMasterFilterIndicatorButton, GetEditorData()->GetUserDefinedFilters())
+							SNew(STextBlock)
+							.Justification(ETextJustify::Center)
+							.TextStyle(FEditorStyle::Get(), "NormalText.Important")
+							.Text(FEditorFontGlyphs::Plus)
+							.Text(LOCTEXT("UpdateResults", "Refresh Results"))
 						]
 					]
 				]
 
-				+ SHorizontalBox::Slot()
-				.HAlign(HAlign_Right)
+				// Save and load
+				+SHorizontalBox::Slot()
+					.HAlign(HAlign_Right)
 				[
-					SNew(SHorizontalBox)
-
-					// Refresh button
-					+ SHorizontalBox::Slot()
-					.AutoWidth()
-					.Padding(0.f, 0.f, 2.f, 0.f)
-					[
-						SNew(SButton)
-						.ButtonStyle(FEditorStyle::Get(), "FlatButton.Success")
-						.ForegroundColor(FSlateColor::UseForeground())
-						.OnClicked(this, &SLevelSnapshotsEditorFilters::OnClickUpdateResultsView)
-						[
-							SNew(SHorizontalBox)
-							+SHorizontalBox::Slot()
-								.AutoWidth()
-								.VAlign(VAlign_Center)
-							[
-								SNew(STextBlock)
-								.Justification(ETextJustify::Center)
-								.TextStyle(FEditorStyle::Get(), "NormalText.Important")
-								.Text(FEditorFontGlyphs::Plus)
-								.Text(LOCTEXT("UpdateResults", "Refresh Results"))
-							]
-						]
-					]
-
-					// Save and load
-					+SHorizontalBox::Slot()
-						.HAlign(HAlign_Right)
-					[
-						SNew(SSaveAndLoadFilters, GetEditorData())
-					]
+					SNew(SSaveAndLoadFilters, GetEditorData())
 				]
 			]
 			
@@ -287,10 +264,6 @@ void SLevelSnapshotsEditorFilters::Construct(const FArguments& InArgs, const TSh
 
 		const TArray<UConjunctionFilter*>* AndFilters = &UserDefinedFilters->GetChildren();
 		FilterRowsList->SetTreeItemsSource(AndFilters);
-
-		// Update master filter
-		MasterFilterIndicatorButton->SetUserDefinedFilters(UserDefinedFilters);
-		UserDefinedFilters->SetEditorFilterBehaviorFromChild();
 		
 		GetEditorData()->SetEditedFilter(TOptional<UNegatableFilter*>());
 	});
@@ -327,25 +300,10 @@ void SLevelSnapshotsEditorFilters::RemoveFilter(UConjunctionFilter* FilterToRemo
 {
 	UDisjunctiveNormalFormFilter* UserDefinedFilters = GetEditorData()->GetUserDefinedFilters();
 	UserDefinedFilters->RemoveConjunction(FilterToRemove);
-	UserDefinedFilters->SetEditorFilterBehaviorFromChild();
 
 	RefreshGroups();
 }
 
-void SLevelSnapshotsEditorFilters::SetMasterFilterBehaviorFromChildRow(const EEditorFilterBehavior InFilterBehavior)
-{
-	ULevelSnapshotsEditorData* EditorData = GetEditorData();
-	if (!ensure(EditorData))
-	{
-		return;
-	}
-
-	if (UDisjunctiveNormalFormFilter* UserDefinedFilters = EditorData->GetUserDefinedFilters())
-	{
-		UserDefinedFilters->SetEditorFilterBehaviorFromChild(InFilterBehavior);
-	}
-}
-
 FReply SLevelSnapshotsEditorFilters::OnClickUpdateResultsView()
 {
 	ULevelSnapshotsEditorData* EditorData = GetEditorData();
diff --git a/Engine/Plugins/VirtualProduction/LevelSnapshots/Source/LevelSnapshotsEditor/Private/Views/Filter/SLevelSnapshotsEditorFilters.h b/Engine/Plugins/VirtualProduction/LevelSnapshots/Source/LevelSnapshotsEditor/Private/Views/Filter/SLevelSnapshotsEditorFilters.h
--- a/Engine/Plugins/VirtualProduction/LevelSnapshots/Source/LevelSnapshotsEditor/Private/Views/Filter/SLevelSnapshotsEditorFilters.h
+++ b/Engine/Plugins/VirtualProduction/LevelSnapshots/Source/LevelSnapshotsEditor/Private/Views/Filter/SLevelSnapshotsEditorFilters.h
@@ -16,7 +16,7 @@ class ULevelSnapshotsEditorData;
 class ULevelSnapshotFilter;
 class UConjunctionFilter;
 
-enum class EEditorFilterBehavior : uint8;
+enum class EFilterBehavior : uint8;
 
 /* Contents of filters tab */
 class SLevelSnapshotsEditorFilters : public SCompoundWidget
@@ -38,8 +38,6 @@ public:
 	
 	void RemoveFilter(UConjunctionFilter* FilterToRemove);
 
-	void SetMasterFilterBehaviorFromChildRow(const EEditorFilterBehavior InFilterBehavior);
-
 private:
 
 	FReply OnClickUpdateResultsView();
@@ -64,6 +62,4 @@ private:
 
 	TWeakPtr<FLevelSnapshotsEditorFilters> FiltersModelPtr;
 	TWeakObjectPtr<ULevelSnapshotsEditorData> EditorDataPtr;
-
-	TSharedPtr<SMasterFilterIndicatorButton> MasterFilterIndicatorButton;
 };
diff --git a/Engine/Plugins/VirtualProduction/LevelSnapshots/Source/LevelSnapshotsEditor/Private/Widgets/Filter/SFilterCheckBox.cpp b/Engine/Plugins/VirtualProduction/LevelSnapshots/Source/LevelSnapshotsEditor/Private/Widgets/Filter/SFilterCheckBox.cpp
--- a/Engine/Plugins/VirtualProduction/LevelSnapshots/Source/LevelSnapshotsEditor/Private/Widgets/Filter/SFilterCheckBox.cpp
+++ b/Engine/Plugins/VirtualProduction/LevelSnapshots/Source/LevelSnapshotsEditor/Private/Widgets/Filter/SFilterCheckBox.cpp
@@ -1,10 +1,8 @@
 ﻿// Copyright Epic Games, Inc. All Rights Reserved.
 
-#include "SFilterCheckBox.h"
-
+#include "Widgets/Filter/SFilterCheckBox.h"
 
 #include "EditorStyleSet.h"
-#include "Widgets/Layout/SBorder.h"
 #include "Widgets/SBoxPanel.h"
 #include "Widgets/Images/SImage.h"
 
diff --git a/Engine/Plugins/VirtualProduction/LevelSnapshots/Source/LevelSnapshotsEditor/Private/Widgets/Filter/SMasterFilterIndicatorButton.cpp b/Engine/Plugins/VirtualProduction/LevelSnapshots/Source/LevelSnapshotsEditor/Private/Widgets/Filter/SMasterFilterIndicatorButton.cpp
deleted file mode 100644
diff --git a/Engine/Plugins/VirtualProduction/LevelSnapshots/Source/LevelSnapshotsEditor/Private/Widgets/Filter/SMasterFilterIndicatorButton.h b/Engine/Plugins/VirtualProduction/LevelSnapshots/Source/LevelSnapshotsEditor/Private/Widgets/Filter/SMasterFilterIndicatorButton.h
deleted file mode 100644
diff --git a/Engine/Plugins/VirtualProduction/LevelSnapshots/Source/LevelSnapshotsEditor/Private/Widgets/Filter/SRowFilterIndicatorButton.cpp b/Engine/Plugins/VirtualProduction/LevelSnapshots/Source/LevelSnapshotsEditor/Private/Widgets/Filter/SRowFilterIndicatorButton.cpp
deleted file mode 100644
diff --git a/Engine/Plugins/VirtualProduction/LevelSnapshots/Source/LevelSnapshotsEditor/Private/Widgets/Filter/SRowFilterIndicatorButton.h b/Engine/Plugins/VirtualProduction/LevelSnapshots/Source/LevelSnapshotsEditor/Private/Widgets/Filter/SRowFilterIndicatorButton.h
deleted file mode 100644
diff --git a/Engine/Plugins/VirtualProduction/LevelSnapshots/Source/LevelSnapshotsEditor/Private/Widgets/SLevelSnapshotsEditorFilter.cpp b/Engine/Plugins/VirtualProduction/LevelSnapshots/Source/LevelSnapshotsEditor/Private/Widgets/SLevelSnapshotsEditorFilter.cpp
--- a/Engine/Plugins/VirtualProduction/LevelSnapshots/Source/LevelSnapshotsEditor/Private/Widgets/SLevelSnapshotsEditorFilter.cpp
+++ b/Engine/Plugins/VirtualProduction/LevelSnapshots/Source/LevelSnapshotsEditor/Private/Widgets/SLevelSnapshotsEditorFilter.cpp
@@ -73,13 +73,13 @@ void SLevelSnapshotsEditorFilter::Construct(const FArguments& InArgs, const TWea
 				{
 					if (SnapshotFilter.IsValid())
 					{
-						switch (SnapshotFilter->GetEditorFilterBehavior())
+						switch (SnapshotFilter->GetFilterBehavior())
 	                    {
-	                        case EEditorFilterBehavior::DoNotNegate:
+	                        case EFilterBehavior::DoNotNegate:
 								return LOCTEXT("DoNotNegate", "Filter result is neither negated nor ignored. Click to toggle.");
-	                        case EEditorFilterBehavior::Negate:
+	                        case EFilterBehavior::Negate:
 								return LOCTEXT("Negate", "Filter result is negated. Click to toggle.");
-	                        case EEditorFilterBehavior::Ignore:
+	                        case EFilterBehavior::Ignore:
 								return LOCTEXT("Ignore", "Filter result is ignored. Click to toggle.");
 	                    }
 					}
@@ -97,13 +97,13 @@ void SLevelSnapshotsEditorFilter::Construct(const FArguments& InArgs, const TWea
 						return FLinearColor::Black;
 					}
 
-					switch (SnapshotFilter->GetEditorFilterBehavior())
+					switch (SnapshotFilter->GetFilterBehavior())
 					{
-						case EEditorFilterBehavior::DoNotNegate:
+						case EFilterBehavior::DoNotNegate:
 							return FLinearColor::Green;
-						case EEditorFilterBehavior::Negate:
+						case EFilterBehavior::Negate:
 							return FLinearColor::Red;
-						case EEditorFilterBehavior::Ignore:
+						case EFilterBehavior::Ignore:
 							return FLinearColor::Gray;
 						default: 
 							return FLinearColor::Black;
@@ -194,7 +194,7 @@ FReply SLevelSnapshotsEditorFilter::OnNegateFilter()
 {
 	if (ensure(SnapshotFilter.IsValid()))
 	{
-		SnapshotFilter->IncrementEditorFilterBehavior();
+		SnapshotFilter->IncrementFilterBehaviour();
 	}
 	return FReply::Handled();
 }
diff --git a/Engine/Plugins/VirtualProduction/LevelSnapshots/Source/LevelSnapshotsEditor/Private/Widgets/SLevelSnapshotsEditorFilterRow.cpp b/Engine/Plugins/VirtualProduction/LevelSnapshots/Source/LevelSnapshotsEditor/Private/Widgets/SLevelSnapshotsEditorFilterRow.cpp
--- a/Engine/Plugins/VirtualProduction/LevelSnapshots/Source/LevelSnapshotsEditor/Private/Widgets/SLevelSnapshotsEditorFilterRow.cpp
+++ b/Engine/Plugins/VirtualProduction/LevelSnapshots/Source/LevelSnapshotsEditor/Private/Widgets/SLevelSnapshotsEditorFilterRow.cpp
@@ -5,7 +5,6 @@
 #include "ConjunctionFilter.h"
 #include "FavoriteFilterDragDrop.h"
 #include "LevelSnapshotsEditorStyle.h"
-#include "SRowFilterIndicatorButton.h"
 #include "SLevelSnapshotsEditorFilterList.h"
 #include "SLevelSnapshotsEditorFilters.h"
 
@@ -65,47 +64,44 @@ void SLevelSnapshotsEditorFilterRow::Construct(
 			.HAlign(HAlign_Fill)
 			.FillWidth(1.f)
 			[
-				SNew(SRowFilterIndicatorButton, InManagedFilter)
+				SNew(SBorder)
+				.Padding(FMargin(5.0f, 5.f))
+				.BorderImage(FLevelSnapshotsEditorStyle::GetBrush("LevelSnapshotsEditor.GroupBorder"))
 				[
-					SNew(SBorder)
-					.Padding(FMargin(5.0f, 5.f))
-					.BorderImage(FLevelSnapshotsEditorStyle::GetBrush("LevelSnapshotsEditor.GroupBorder"))
+					SNew(SHorizontalBox)
+
+					// Filters
+					+ SHorizontalBox::Slot()
+					.HAlign(HAlign_Fill)
+					.Padding(5.f, 5.f)
+					.FillWidth(1.f)
 					[
-						SNew(SHorizontalBox)
+						SAssignNew(FilterList, SLevelSnapshotsEditorFilterList, InManagedFilter, InEditorFilters->GetFiltersModel().ToSharedRef())
+					]
 
-						// Filters
-						+ SHorizontalBox::Slot()
-						.HAlign(HAlign_Fill)
-						.Padding(5.f, 5.f)
-						.FillWidth(1.f)
-						[
-							SAssignNew(FilterList, SLevelSnapshotsEditorFilterList, InManagedFilter, InEditorFilters->GetFiltersModel().ToSharedRef())
-						]
+					+ SHorizontalBox::Slot()
+					.VAlign(VAlign_Top)
+					.AutoWidth()
+					[
+						SNew(SHorizontalBox)
 
+						// Remove Button
 						+ SHorizontalBox::Slot()
-						.VAlign(VAlign_Top)
+						.Padding(0.f, 0.f)
 						.AutoWidth()
 						[
-							SNew(SHorizontalBox)
-
-							// Remove Button
-							+ SHorizontalBox::Slot()
-							.Padding(0.f, 0.f)
-							.AutoWidth()
+							SNew(SButton)
+							.OnClicked_Lambda([this]()
+							{
+								OnClickRemoveRow.ExecuteIfBound(SharedThis(this));
+								return FReply::Handled();
+							})
+							.ButtonStyle(FLevelSnapshotsEditorStyle::Get(), "LevelSnapshotsEditor.RemoveFilterButton")
 							[
-								SNew(SButton)
-								.OnClicked_Lambda([this]()
-								{
-									OnClickRemoveRow.ExecuteIfBound(SharedThis(this));
-									return FReply::Handled();
-								})
-								.ButtonStyle(FLevelSnapshotsEditorStyle::Get(), "LevelSnapshotsEditor.RemoveFilterButton")
-								[
-									SNew(STextBlock)
-									.TextStyle(FLevelSnapshotsEditorStyle::Get(), "LevelSnapshotsEditor.Button.TextStyle")
-									.Font(FEditorStyle::Get().GetFontStyle("FontAwesome.10"))
-									.Text(FText::FromString(FString(TEXT("\xf00d"))) /*fa-times*/)
-								]
+								SNew(STextBlock)
+								.TextStyle(FLevelSnapshotsEditorStyle::Get(), "LevelSnapshotsEditor.Button.TextStyle")
+								.Font(FEditorStyle::Get().GetFontStyle("FontAwesome.10"))
+								.Text(FText::FromString(FString(TEXT("\xf00d"))) /*fa-times*/)
 							]
 						]
 					]
