Fix post process property missing

diff --git a/Engine/Plugins/VirtualProduction/LevelSnapshots/Source/LevelSnapshots/Private/Tests/ApplyToWorldTests.cpp b/Engine/Plugins/VirtualProduction/LevelSnapshots/Source/LevelSnapshots/Private/Tests/ApplyToWorldTests.cpp
--- a/Engine/Plugins/VirtualProduction/LevelSnapshots/Source/LevelSnapshots/Private/Tests/ApplyToWorldTests.cpp
+++ b/Engine/Plugins/VirtualProduction/LevelSnapshots/Source/LevelSnapshots/Private/Tests/ApplyToWorldTests.cpp
@@ -5,7 +5,6 @@
 
 #include "Components/PointLightComponent.h"
 #include "Engine/PointLight.h"
-#include "Engine/PostProcessVolume.h"
 #include "Engine/StaticMeshActor.h"
 #include "Misc/AutomationTest.h"
 
@@ -19,29 +18,6 @@ bool FRestoreSimpleProperties::RunTest(const FString& Parameters)
 	const FVector StartScale(1.f, 1.f, 2.f);
 	const FVector EndScale(2.f, 3.f, -2.f);
 
-	const FPostProcessSettings StartSettings = []()
-	{
-		FPostProcessSettings Result;
-		Result.bOverride_ColorSaturation = true;
-		Result.ColorSaturation = FVector4(0.5f, 0.5f, 0.5f, 0.5f);
-		Result.bOverride_ReflectionsType = true;
-		Result.ReflectionsType = EReflectionsType::ScreenSpace;
-		Result.bOverride_TranslucencyType = true;
-		Result.TranslucencyType = ETranslucencyType::Raster;
-		return Result;
-	}();
-	const FPostProcessSettings EndSettings = []()
-	{
-		FPostProcessSettings Result;
-		Result.bOverride_ColorSaturation = true;
-		Result.ColorSaturation = FVector4(1.5f, 1.5f, 1.5f, 1.5f);
-		Result.bOverride_ReflectionsType = true;
-		Result.ReflectionsType = EReflectionsType::RayTracing;
-		Result.bOverride_TranslucencyType = false;
-		Result.TranslucencyType = ETranslucencyType::RayTracing;
-		return Result;
-	}();
-
 	const float StartRadius = 200.f;
 	const float EndRadius = 500.f;
 	const FColor StartColor = FColor::Red;
@@ -50,7 +26,6 @@ bool FRestoreSimpleProperties::RunTest(const FString& Parameters)
 	const uint32 bEndCastShaodws = false;
 	
 	AStaticMeshActor* StaticMesh	= nullptr;
-	APostProcessVolume* PostProcess = nullptr;
 	APointLight* PointLight			= nullptr;
 	
 	FSnapshotTestRunner()
@@ -59,9 +34,6 @@ bool FRestoreSimpleProperties::RunTest(const FString& Parameters)
 			StaticMesh = World->SpawnActor<AStaticMeshActor>(StartLocation, StartRotation);
 			StaticMesh->SetActorScale3D(StartScale);
 
-			PostProcess = World->SpawnActor<APostProcessVolume>();
-			PostProcess->Settings = StartSettings;
-
 			PointLight = World->SpawnActor<APointLight>(StartLocation, StartRotation);
 			PointLight->PointLightComponent->AttenuationRadius = StartRadius;
 			PointLight->PointLightComponent->LightColor = StartColor;
@@ -74,9 +46,6 @@ bool FRestoreSimpleProperties::RunTest(const FString& Parameters)
 			StaticMesh->SetActorLocationAndRotation(EndLocation, EndRotation);
 			StaticMesh->SetActorScale3D(EndScale);
 
-			PostProcess->Settings = EndSettings;
-			PostProcess->SetActorLocationAndRotation(StartLocation, StartRotation);
-
 			PointLight->PointLightComponent->AttenuationRadius = EndRadius;
 			PointLight->PointLightComponent->LightColor = EndColor;
 			PointLight->PointLightComponent->CastShadows = bEndCastShaodws;
@@ -89,16 +58,6 @@ bool FRestoreSimpleProperties::RunTest(const FString& Parameters)
 			TestEqual("Static Mesh Rotation", StaticMesh->GetActorRotation(), StartRotation);
 			TestEqual("Static Mesh Scale", StaticMesh->GetActorScale3D(), StartScale);
 
-			TestEqual("Post Process bOverride_ColorSaturation", PostProcess->Settings.bOverride_ColorSaturation, StartSettings.bOverride_ColorSaturation);
-			TestEqual("Post Process ColorSaturation", PostProcess->Settings.ColorSaturation, StartSettings.ColorSaturation);
-			TestEqual("Post Process bOverride_ReflectionsType", PostProcess->Settings.bOverride_ReflectionsType, StartSettings.bOverride_ReflectionsType);
-			TestEqual("Post Process ReflectionsType", PostProcess->Settings.ReflectionsType, StartSettings.ReflectionsType);
-			TestEqual("Post Process bOverride_TranslucencyType", PostProcess->Settings.bOverride_TranslucencyType, StartSettings.bOverride_TranslucencyType);
-			TestEqual("Post Process TranslucencyType", PostProcess->Settings.TranslucencyType, StartSettings.TranslucencyType);
-			
-			TestEqual("Post Process Location", PostProcess->GetActorLocation(), FVector(0.f));
-			TestEqual("Post Process Rotation", PostProcess->GetActorRotation(), FRotator(0.f));
-
 			TestEqual("Point Light Radius", PointLight->PointLightComponent->AttenuationRadius, StartRadius);
 			TestEqual("Point Light Colour", PointLight->PointLightComponent->LightColor, StartColor);
 			TestEqual("Point Light Cast Shadows", PointLight->PointLightComponent->CastShadows, bStartCastShadows);
