Fix ambigous function call

diff --git a/Engine/Plugins/VirtualProduction/LevelSnapshots/Source/LevelSnapshots/Private/Tests/ApplyToWorldTests.cpp b/Engine/Plugins/VirtualProduction/LevelSnapshots/Source/LevelSnapshots/Private/Tests/ApplyToWorldTests.cpp
--- a/Engine/Plugins/VirtualProduction/LevelSnapshots/Source/LevelSnapshots/Private/Tests/ApplyToWorldTests.cpp
+++ b/Engine/Plugins/VirtualProduction/LevelSnapshots/Source/LevelSnapshots/Private/Tests/ApplyToWorldTests.cpp
@@ -5,7 +5,9 @@
 
 #include "Components/PointLightComponent.h"
 #include "Engine/PointLight.h"
+#include "Engine/StaticMesh.h"
 #include "Engine/StaticMeshActor.h"
+#include "Materials/MaterialInterface.h"
 #include "Misc/AutomationTest.h"
 
 IMPLEMENT_SIMPLE_AUTOMATION_TEST(FRestoreSimpleProperties, "VirtualProduction.LevelSnapshots.Snapshot.RestoreSimpleProperties", (EAutomationTestFlags::ApplicationContextMask | EAutomationTestFlags::EngineFilter));
@@ -54,13 +56,13 @@ bool FRestoreSimpleProperties::RunTest(const FString& Parameters)
 
 		.RunTest([&]()
 		{
-			TestEqual("Static Mesh Location", StaticMesh->GetActorLocation(), StartLocation);
-			TestEqual("Static Mesh Rotation", StaticMesh->GetActorRotation(), StartRotation);
-			TestEqual("Static Mesh Scale", StaticMesh->GetActorScale3D(), StartScale);
+			TestEqual(TEXT("Static Mesh Location"), StaticMesh->GetActorLocation(), StartLocation);
+			TestEqual(TEXT("Static Mesh Rotation"), StaticMesh->GetActorRotation(), StartRotation);
+			TestEqual(TEXT("Static Mesh Scale"), StaticMesh->GetActorScale3D(), StartScale);
 
-			TestEqual("Point Light Radius", PointLight->PointLightComponent->AttenuationRadius, StartRadius);
-			TestEqual("Point Light Colour", PointLight->PointLightComponent->LightColor, StartColor);
-			TestEqual("Point Light Cast Shadows", PointLight->PointLightComponent->CastShadows, bStartCastShadows);
+			TestEqual(TEXT("Point Light Radius"), PointLight->PointLightComponent->AttenuationRadius, StartRadius);
+			TestEqual(TEXT("Point Light Colour"), PointLight->PointLightComponent->LightColor, StartColor);
+			TestEqual(TEXT("Point Light Cast Shadows"), PointLight->PointLightComponent->CastShadows, bStartCastShadows);
 		});
 	
 	return true;
@@ -122,14 +124,14 @@ bool FRestoreReferenceProperties::RunTest(const FString& Parameters)
 
 		.RunTest([&]()
 		{
-			TestTrue("Cyclic Reference Restored", FirstActor->HasObjectReference(SecondActor));
-			TestTrue("Reference Not Lost", SecondActor->HasObjectReference(FirstActor));
-			TestTrue("Restore Pointing To Self", PointToSelfActor->HasObjectReference(PointToSelfActor));
-			TestTrue("World > External > World", FromWorldToExternal->HasObjectReference(FirstActor));
-			TestTrue("External > World > External", FromExternalToWorld->HasObjectReference(FromExternalToWorld->CubeMesh));
-
-			TestEqual("Mesh", MaterialAndMesh->InstancedMeshComponent->GetStaticMesh(), MaterialAndMesh->CubeMesh);
-			TestEqual("Material", MaterialAndMesh->InstancedMeshComponent->GetMaterial(0), MaterialAndMesh->GradientLinearMaterial);
+			TestTrue(TEXT("Cyclic Reference Restored"), FirstActor->HasObjectReference(SecondActor));
+			TestTrue(TEXT("Reference Not Lost"), SecondActor->HasObjectReference(FirstActor));
+			TestTrue(TEXT("Restore Pointing To Self"), PointToSelfActor->HasObjectReference(PointToSelfActor));
+			TestTrue(TEXT("World > External > World"), FromWorldToExternal->HasObjectReference(FirstActor));
+			TestTrue(TEXT("External > World > External"), FromExternalToWorld->HasObjectReference(FromExternalToWorld->CubeMesh));
+
+			TestEqual(TEXT("Mesh"), MaterialAndMesh->InstancedMeshComponent->GetStaticMesh(), MaterialAndMesh->CubeMesh);
+			TestEqual(TEXT("Material"), MaterialAndMesh->InstancedMeshComponent->GetMaterial(0), MaterialAndMesh->GradientLinearMaterial);
 			
 		});
 
@@ -165,15 +167,15 @@ bool FInstancedStaticMesh::RunTest(const FString& Parameters)
 
 		.RunTest([&]()
 		{
-			TestEqual("Instance Count", MaterialAndMesh->InstancedMeshComponent->GetInstanceCount(), 2);
+			TestEqual(TEXT("Instance Count"), MaterialAndMesh->InstancedMeshComponent->GetInstanceCount(), 2);
 
 			FTransform ActualFirstTransform, ActualSecondTransform;
 			MaterialAndMesh->InstancedMeshComponent->GetInstanceTransform(0, ActualFirstTransform);
 			MaterialAndMesh->InstancedMeshComponent->GetInstanceTransform(1, ActualSecondTransform);
 
 			// TODO: Investigate rotation not working correcty... it is some kind of math problem that causes angle to flip, i.e. 80 and 100 represent the same angle but are obviously not equal floats...
-			TestTrue("Instance 1", ActualFirstTransform.GetLocation().Equals(StartFirstTransform.GetLocation()) && ActualFirstTransform.GetScale3D().Equals(StartFirstTransform.GetScale3D()));
-			TestTrue("Instance 2", ActualSecondTransform.GetLocation().Equals(StartSecondTransform.GetLocation())&& ActualSecondTransform.GetScale3D().Equals(StartSecondTransform.GetScale3D()));
+			TestTrue(TEXT("Instance 1"), ActualFirstTransform.GetLocation().Equals(StartFirstTransform.GetLocation()) && ActualFirstTransform.GetScale3D().Equals(StartFirstTransform.GetScale3D()));
+			TestTrue(TEXT("Instance 2"), ActualSecondTransform.GetLocation().Equals(StartSecondTransform.GetLocation())&& ActualSecondTransform.GetScale3D().Equals(StartSecondTransform.GetScale3D()));
 		});
 
 	return true;
@@ -200,8 +202,8 @@ bool FSkipTransientAndDeprecatedProperties::RunTest(const FString& Parameters)
 
 		.RunTest([&]()
 		{
-			TestEqual("Skip Deprecated Property", TestActor->DeprecatedProperty_DEPRECATED, 500);
-			TestEqual("Skip Transient Property", TestActor->TransientProperty, 500);
+			TestEqual(TEXT("Skip Deprecated Property"), TestActor->DeprecatedProperty_DEPRECATED, 500);
+			TestEqual(TEXT("Skip Transient Property"), TestActor->TransientProperty, 500);
 		});
 
 	return true;
