Fix horde issue 75052

diff --git a/Engine/Plugins/VirtualProduction/LevelSnapshots/Source/LevelSnapshotsEditor/Private/Views/Filter/SLevelSnapshotsEditorFilters.cpp b/Engine/Plugins/VirtualProduction/LevelSnapshots/Source/LevelSnapshotsEditor/Private/Views/Filter/SLevelSnapshotsEditorFilters.cpp
--- a/Engine/Plugins/VirtualProduction/LevelSnapshots/Source/LevelSnapshotsEditor/Private/Views/Filter/SLevelSnapshotsEditorFilters.cpp
+++ b/Engine/Plugins/VirtualProduction/LevelSnapshots/Source/LevelSnapshotsEditor/Private/Views/Filter/SLevelSnapshotsEditorFilters.cpp
@@ -123,15 +123,15 @@ void SLevelSnapshotsEditorFilters::Construct(const FArguments& InArgs, const TSh
 
 	// Create a property view
 	FPropertyEditorModule& EditModule = FModuleManager::Get().GetModuleChecked<FPropertyEditorModule>("PropertyEditor");
-	FDetailsViewArgs DetailsViewArgs(
-		/*bUpdateFromSelection=*/ false,
-		/*bLockable=*/ false, 
-		/*bAllowSearch=*/ true,
-		FDetailsViewArgs::HideNameArea,
-		/*bHideSelectionTip=*/ true,
-		/*InNotifyHook=*/ nullptr,
-		/*InSearchInitialKeyFocus=*/ false,
-		/*InViewIdentifier=*/ NAME_None);
+	FDetailsViewArgs DetailsViewArgs;
+	DetailsViewArgs.bUpdatesFromSelection = false;
+	DetailsViewArgs.bLockable = false;
+	DetailsViewArgs.bAllowSearch = true;
+	DetailsViewArgs.NameAreaSettings = FDetailsViewArgs::HideNameArea;
+	DetailsViewArgs.bHideSelectionTip = true;
+	DetailsViewArgs.NotifyHook = nullptr;
+	DetailsViewArgs.bSearchInitialKeyFocus = false;
+	DetailsViewArgs.ViewIdentifier = NAME_None;
 	DetailsViewArgs.DefaultsOnlyVisibility = EEditDefaultsOnlyNodeVisibility::Automatic;
 	
 	FilterDetailsView = EditModule.CreateDetailView(DetailsViewArgs);
