Call pre and edit post edit change

diff --git a/Engine/Plugins/VirtualProduction/LevelSnapshots/Source/LevelSnapshots/Private/Archive/ApplySnapshotDataArchiveV2.cpp b/Engine/Plugins/VirtualProduction/LevelSnapshots/Source/LevelSnapshots/Private/Archive/ApplySnapshotDataArchiveV2.cpp
--- a/Engine/Plugins/VirtualProduction/LevelSnapshots/Source/LevelSnapshots/Private/Archive/ApplySnapshotDataArchiveV2.cpp
+++ b/Engine/Plugins/VirtualProduction/LevelSnapshots/Source/LevelSnapshots/Private/Archive/ApplySnapshotDataArchiveV2.cpp
@@ -12,6 +12,7 @@
 #include "Components/ActorComponent.h"
 #include "Engine/World.h"
 #include "GameFramework/Actor.h"
+#include "Misc/ScopeExit.h"
 #include "Serialization/ObjectReader.h"
 #include "Serialization/ObjectWriter.h"
 #include "UObject/TextProperty.h"
@@ -143,7 +144,15 @@ namespace
 }
 
 void FApplySnapshotDataArchiveV2::ApplyToExistingEditorWorldObject(FObjectSnapshotData& InObjectData, FWorldSnapshotData& InSharedData, UObject* InOriginalObject, UObject* InDeserializedVersion, const FPropertySelectionMap& InSelectionMapForResolvingSubobjects, const FPropertySelection& InSelectionSet)
-{	
+{
+#if WITH_EDITOR
+	InOriginalObject->PreEditChange(nullptr);
+	ON_SCOPE_EXIT
+	{
+		InOriginalObject->PostEditChange();	
+	};
+#endif
+	
 	// Step 1: Serialise  properties that were different from CDO at time of snapshotting and that are still different from CDO
 	FApplySnapshotDataArchiveV2 ApplySavedData(InObjectData, InSharedData, InOriginalObject, InSelectionMapForResolvingSubobjects, &InSelectionSet);
 	InOriginalObject->Serialize(ApplySavedData);
@@ -157,6 +166,14 @@ void FApplySnapshotDataArchiveV2::ApplyToExistingEditorWorldObject(FObjectSnapsh
 
 void FApplySnapshotDataArchiveV2::ApplyToRecreatedEditorWorldObject(FObjectSnapshotData& InObjectData, FWorldSnapshotData& InSharedData, UObject* InOriginalObject, UObject* InDeserializedVersion, const FPropertySelectionMap& InSelectionMapForResolvingSubobjects)
 {
+#if WITH_EDITOR
+	InOriginalObject->PreEditChange(nullptr);
+	ON_SCOPE_EXIT
+	{
+		InOriginalObject->PostEditChange();	
+	};
+#endif
+	
 	// Apply all properties that we saved into the target actor.
 	// We assume that InOriginalObject was already created with the snapshot CDO as template: we do not need Step 2 from ApplyToExistingWorldObject.
 	FApplySnapshotDataArchiveV2 ApplySavedData(InObjectData, InSharedData, InOriginalObject, InSelectionMapForResolvingSubobjects, {});
@@ -184,37 +201,6 @@ void FApplySnapshotDataArchiveV2::PushSerializedProperty(FProperty* InProperty,
 	PropertiesLeftToSerialize.RemoveProperty(GetSerializedPropertyChain(), InProperty);
 	
 	Super::PushSerializedProperty(InProperty, bIsEditorOnlyProperty);
-
-#if WITH_EDITOR
-	DECLARE_SCOPE_CYCLE_COUNTER(TEXT("PreEditChange"), STAT_PreEditChange, STATGROUP_LevelSnapshots);
-	if (FStructProperty* StructProperty = CastField<FStructProperty>(InProperty))
-	{
-		// TODO: chain event
-	}
-	else
-	{
-		OriginalObject->PreEditChange(InProperty);
-	}
-#endif
-	
-}
-
-void FApplySnapshotDataArchiveV2::PopSerializedProperty(FProperty* InProperty, const bool bIsEditorOnlyProperty)
-{
-	Super::PopSerializedProperty(InProperty, bIsEditorOnlyProperty);
-
-#if WITH_EDITOR
-	DECLARE_SCOPE_CYCLE_COUNTER(TEXT("PostEditChange"), STAT_PostEditChange, STATGROUP_LevelSnapshots);
-	if (FStructProperty* StructProperty = CastField<FStructProperty>(InProperty))
-	{
-		// TODO: chain event
-	}
-	else
-	{
-		FPropertyChangedEvent ChangeEvent(InProperty);
-		OriginalObject->PostEditChangeProperty(ChangeEvent);
-	}
-#endif
 }
 
 UObject* FApplySnapshotDataArchiveV2::ResolveObjectDependency(int32 ObjectIndex) const
diff --git a/Engine/Plugins/VirtualProduction/LevelSnapshots/Source/LevelSnapshots/Private/Archive/ApplySnapshotDataArchiveV2.h b/Engine/Plugins/VirtualProduction/LevelSnapshots/Source/LevelSnapshots/Private/Archive/ApplySnapshotDataArchiveV2.h
--- a/Engine/Plugins/VirtualProduction/LevelSnapshots/Source/LevelSnapshots/Private/Archive/ApplySnapshotDataArchiveV2.h
+++ b/Engine/Plugins/VirtualProduction/LevelSnapshots/Source/LevelSnapshots/Private/Archive/ApplySnapshotDataArchiveV2.h
@@ -19,7 +19,6 @@ public:
 	//~ Begin FTakeSnapshotArchiveV2 Interface
 	virtual bool ShouldSkipProperty(const FProperty* InProperty) const override;
 	virtual void PushSerializedProperty(FProperty* InProperty, const bool bIsEditorOnlyProperty) override;
-	virtual void PopSerializedProperty(FProperty* InProperty, const bool bIsEditorOnlyProperty) override;
 	//~ End FTakeSnapshotArchiveV2 Interface
 
 protected:
