Results:
-Merged two info text blocks into one. They were two previously because they reflected two separate calculations, but now we do both at the same time. This also has the side effect of truncating the two stacked info text blocks in a similar way. Previously, the top two text blocks overlapped the Restore button when the window wasn't wide enough. Now both stacked text blocks wil be cut off on the left side.
-Added OnMapChanged delegate binding to LevelSnapshotsEditorData so the world and snapshot will be cleared when the map changes.

#[fyi] Jason.Walter

diff --git a/Engine/Plugins/VirtualProduction/LevelSnapshots/Source/LevelSnapshotsEditor/Private/Data/LevelSnapshotsEditorData.cpp b/Engine/Plugins/VirtualProduction/LevelSnapshots/Source/LevelSnapshotsEditor/Private/Data/LevelSnapshotsEditorData.cpp
--- a/Engine/Plugins/VirtualProduction/LevelSnapshots/Source/LevelSnapshotsEditor/Private/Data/LevelSnapshotsEditorData.cpp
+++ b/Engine/Plugins/VirtualProduction/LevelSnapshots/Source/LevelSnapshotsEditor/Private/Data/LevelSnapshotsEditorData.cpp
@@ -3,11 +3,13 @@
 #include "LevelSnapshotsEditorData.h"
 #include "SLevelSnapshotsEditorInput.h"
 
+#include "Editor.h"
+#include "Engine/World.h"
+
 #include "FavoriteFilterContainer.h"
 #include "DisjunctiveNormalFormFilter.h"
 #include "FilterLoader.h"
 #include "FilteredResults.h"
-#include "Engine/World.h"
 
 ULevelSnapshotsEditorData::ULevelSnapshotsEditorData(const FObjectInitializer& ObjectInitializer)
 {
@@ -43,7 +45,13 @@ ULevelSnapshotsEditorData::ULevelSnapshotsEditorData(const FObjectInitializer& O
 
 	OnWorldCleanup = FWorldDelegates::OnWorldCleanup.AddLambda([this](UWorld* World, bool bSessionEnded, bool bCleanupResources)
     {
-        SetActiveSnapshot(nullptr);
+        ClearActiveSnapshot();
+    });
+
+	OnMapOpenedDelegateHandle = FEditorDelegates::OnMapOpened.AddLambda([this](const FString& FileName, bool bAsTemplate)
+    {
+		ClearActiveSnapshot();
+		ClearSelectedWorld();
     });
 }
 
@@ -52,6 +60,7 @@ void ULevelSnapshotsEditorData::BeginDestroy()
 	Super::BeginDestroy();
 	
 	FWorldDelegates::OnWorldCleanup.Remove(OnWorldCleanup);
+	FEditorDelegates::OnMapOpened.Remove(OnMapOpenedDelegateHandle);
 }
 
 void ULevelSnapshotsEditorData::CleanupAfterEditorClose()
@@ -59,6 +68,7 @@ void ULevelSnapshotsEditorData::CleanupAfterEditorClose()
 	OnActiveSnapshotChanged.Clear();
 	OnEditedFiterChanged.Clear();
 	OnUserDefinedFiltersChanged.Clear();
+	OnMapOpenedDelegateHandle.Reset();
 
 	SelectedWorld.Reset();
 	ActiveSnapshot.Reset();
diff --git a/Engine/Plugins/VirtualProduction/LevelSnapshots/Source/LevelSnapshotsEditor/Private/Data/LevelSnapshotsEditorData.h b/Engine/Plugins/VirtualProduction/LevelSnapshots/Source/LevelSnapshotsEditor/Private/Data/LevelSnapshotsEditorData.h
--- a/Engine/Plugins/VirtualProduction/LevelSnapshots/Source/LevelSnapshotsEditor/Private/Data/LevelSnapshotsEditorData.h
+++ b/Engine/Plugins/VirtualProduction/LevelSnapshots/Source/LevelSnapshotsEditor/Private/Data/LevelSnapshotsEditorData.h
@@ -41,9 +41,9 @@ public:
 
 	DECLARE_EVENT_OneParam(ULevelSnapshotsEditorData, FOnActiveSnapshotChanged, const TOptional<ULevelSnapshot*>& /* NewSnapshot */);
 	FOnActiveSnapshotChanged OnActiveSnapshotChanged;
+	FDelegateHandle OnMapOpenedDelegateHandle;
 
 
-	
 	/******************** Selected world ********************/
 	void SetSelectedWorldReference(UWorld* InWorld);
 	void ClearSelectedWorld();
diff --git a/Engine/Plugins/VirtualProduction/LevelSnapshots/Source/LevelSnapshotsEditor/Private/Views/Input/SlevelSnapshotsEditorContextPicker.cpp b/Engine/Plugins/VirtualProduction/LevelSnapshots/Source/LevelSnapshotsEditor/Private/Views/Input/SlevelSnapshotsEditorContextPicker.cpp
--- a/Engine/Plugins/VirtualProduction/LevelSnapshots/Source/LevelSnapshotsEditor/Private/Views/Input/SlevelSnapshotsEditorContextPicker.cpp
+++ b/Engine/Plugins/VirtualProduction/LevelSnapshots/Source/LevelSnapshotsEditor/Private/Views/Input/SlevelSnapshotsEditorContextPicker.cpp
@@ -21,6 +21,7 @@
 SLevelSnapshotsEditorContextPicker::~SLevelSnapshotsEditorContextPicker()
 {
 	FEditorDelegates::OnMapOpened.Remove(OnMapOpenedDelegateHandle);
+	OnMapOpenedDelegateHandle.Reset();
 }
 
 void SLevelSnapshotsEditorContextPicker::Construct(const FArguments& InArgs, const TSharedRef<FLevelSnapshotsEditorViewBuilder>& InBuilder)
diff --git a/Engine/Plugins/VirtualProduction/LevelSnapshots/Source/LevelSnapshotsEditor/Private/Views/Results/SLevelSnapshotsEditorResults.cpp b/Engine/Plugins/VirtualProduction/LevelSnapshots/Source/LevelSnapshotsEditor/Private/Views/Results/SLevelSnapshotsEditorResults.cpp
--- a/Engine/Plugins/VirtualProduction/LevelSnapshots/Source/LevelSnapshotsEditor/Private/Views/Results/SLevelSnapshotsEditorResults.cpp
+++ b/Engine/Plugins/VirtualProduction/LevelSnapshots/Source/LevelSnapshotsEditor/Private/Views/Results/SLevelSnapshotsEditorResults.cpp
@@ -1625,33 +1625,19 @@ void SLevelSnapshotsEditorResults::Construct(const FArguments& InArgs, ULevelSna
 					.HAlign(HAlign_Right)
 					.Padding(FMargin(5.f, 5.f))
 					[
-						SNew(SVerticalBox)
+						SAssignNew(InfoTextBox, SVerticalBox)
+						.Visibility_Lambda([this]()
+						{
+							return EditorDataPtr.IsValid() && EditorDataPtr->GetSelectedWorld() && EditorDataPtr->GetActiveSnapshot() ? 
+								EVisibility::HitTestInvisible : EVisibility::Hidden;
+						})
 
 						+SVerticalBox::Slot()
 						.AutoHeight()
 						[
-							SAssignNew(InfoTextBox, SHorizontalBox)
-							.Visibility_Lambda([this]()
-							{
-								return EditorDataPtr.IsValid() && EditorDataPtr->GetSelectedWorld() && EditorDataPtr->GetActiveSnapshot() ? 
-									EVisibility::HitTestInvisible : EVisibility::Hidden;
-							})
-
-							+SHorizontalBox::Slot()
-							.AutoWidth()
-							[
-								SAssignNew(SelectedActorCountText, STextBlock)
-								.Font(FCoreStyle::GetDefaultFontStyle("Bold", 14))
-								.Justification(ETextJustify::Right)
-							]
-
-							+SHorizontalBox::Slot()
-							.AutoWidth()
-							[
-								SAssignNew(TotalActorCountText, STextBlock)
-								.Font(FCoreStyle::GetDefaultFontStyle("Regular", 14))
-								.Justification(ETextJustify::Right)
-							]
+							SAssignNew(SelectedActorCountText, STextBlock)
+							.Font(FCoreStyle::GetDefaultFontStyle("Bold", 14))
+							.Justification(ETextJustify::Right)
 						]
 
 						+ SVerticalBox::Slot()
@@ -1704,6 +1690,8 @@ SLevelSnapshotsEditorResults::~SLevelSnapshotsEditorResults()
 	
 	OnActiveSnapshotChangedHandle.Reset();
 	OnRefreshResultsHandle.Reset();
+	
+	FEditorDelegates::OnMapOpened.Remove(OnMapOpenedDelegateHandle);
 	OnMapOpenedDelegateHandle.Reset();
 	
 	ResultsSearchBoxPtr.Reset();
@@ -1714,7 +1702,6 @@ SLevelSnapshotsEditorResults::~SLevelSnapshotsEditorResults()
 	SelectedSnapshotNamePtr.Reset();
 
 	SelectedActorCountText.Reset();
-	TotalActorCountText.Reset();
 	MiscActorCountText.Reset();
 
 	DummyRow.Reset();
@@ -1909,14 +1896,8 @@ void SLevelSnapshotsEditorResults::UpdateSnapshotInformationText()
 
 	if (SelectedActorCountText.IsValid())
 	{
-		SelectedActorCountText->SetText(FText::Format(LOCTEXT("ResultsRestoreInfoFormatSelectedActorCount", "{0} actor(s)"),
-			FText::AsNumber(SelectedActorCount)));
-	}
-
-	if (TotalActorCountText.IsValid())
-	{
-		TotalActorCountText->SetText(FText::Format(LOCTEXT("ResultsRestoreInfoFormatTotalActorCount", " (of {0} in snapshot) will be restored"),
-			FText::AsNumber(TotalPassingActorCount)));
+		SelectedActorCountText->SetText(FText::Format(LOCTEXT("ResultsRestoreInfoFormatSelectedActorCount", "{0} actor(s) (of {1} in snapshot) will be restored"),
+			FText::AsNumber(SelectedActorCount), FText::AsNumber(TotalPassingActorCount)));
 	}
 
 	if (MiscActorCountText.IsValid())
diff --git a/Engine/Plugins/VirtualProduction/LevelSnapshots/Source/LevelSnapshotsEditor/Private/Views/Results/SLevelSnapshotsEditorResults.h b/Engine/Plugins/VirtualProduction/LevelSnapshots/Source/LevelSnapshotsEditor/Private/Views/Results/SLevelSnapshotsEditorResults.h
--- a/Engine/Plugins/VirtualProduction/LevelSnapshots/Source/LevelSnapshotsEditor/Private/Views/Results/SLevelSnapshotsEditorResults.h
+++ b/Engine/Plugins/VirtualProduction/LevelSnapshots/Source/LevelSnapshotsEditor/Private/Views/Results/SLevelSnapshotsEditorResults.h
@@ -385,9 +385,8 @@ private:
 	TSharedPtr<STextBlock> SelectedSnapshotNamePtr;
 
 	// Snapshot Information Text
-	TSharedPtr<SHorizontalBox> InfoTextBox;
+	TSharedPtr<SVerticalBox> InfoTextBox;
 	TSharedPtr<STextBlock> SelectedActorCountText;
-	TSharedPtr<STextBlock> TotalActorCountText;
 	TSharedPtr<STextBlock> MiscActorCountText;
 
 	FText DefaultNameText;
