Crash when recreating selected component

diff --git a/Engine/Plugins/VirtualProduction/LevelSnapshots/Source/LevelSnapshots/Private/Data/Util/Restoration/WorldDataUtil.cpp b/Engine/Plugins/VirtualProduction/LevelSnapshots/Source/LevelSnapshots/Private/Data/Util/Restoration/WorldDataUtil.cpp
--- a/Engine/Plugins/VirtualProduction/LevelSnapshots/Source/LevelSnapshots/Private/Data/Util/Restoration/WorldDataUtil.cpp
+++ b/Engine/Plugins/VirtualProduction/LevelSnapshots/Source/LevelSnapshots/Private/Data/Util/Restoration/WorldDataUtil.cpp
@@ -406,6 +406,24 @@ namespace UE::LevelSnapshots::Private::Internal
 			}
 		}
 	}
+
+#if WITH_EDITOR
+	class FScopedEditorSelectionClearer
+	{
+		FSelectionStateOfLevel SelectionStateOfLevel;
+	public:
+		FScopedEditorSelectionClearer()
+		{
+			GEditor->GetSelectionStateOfLevel(SelectionStateOfLevel);
+			GEditor->SelectNone(true, true, false);
+		}
+		
+		~FScopedEditorSelectionClearer()
+		{
+			GEditor->SetSelectionStateOfLevel(SelectionStateOfLevel);
+		}
+	};
+#endif
 }
 
 void UE::LevelSnapshots::Private::ApplyToWorld(FWorldSnapshotData& WorldData, FSnapshotDataCache& Cache, UWorld* WorldToApplyTo, UPackage* LocalisationSnapshotPackage, const FPropertySelectionMap& PropertiesToSerialize)
@@ -417,6 +435,8 @@ void UE::LevelSnapshots::Private::ApplyToWorld(FWorldSnapshotData& WorldData, FS
 	PreloadClassesForRestore(WorldData, PropertiesToSerialize);
 	
 #if WITH_EDITOR
+	// Deleting components while they are selected can cause a crash
+	FScopedEditorSelectionClearer RestoreSelection;
 	FScopedTransaction Transaction(FText::FromString("Loading Level Snapshot."));
 #endif
 
