Fix filters not getting marked dirty when changed

diff --git a/Engine/Plugins/VirtualProduction/LevelSnapshots/Source/LevelSnapshotsEditor/Private/Data/LevelSnapshotsEditorData.cpp b/Engine/Plugins/VirtualProduction/LevelSnapshots/Source/LevelSnapshotsEditor/Private/Data/LevelSnapshotsEditorData.cpp
--- a/Engine/Plugins/VirtualProduction/LevelSnapshots/Source/LevelSnapshotsEditor/Private/Data/LevelSnapshotsEditorData.cpp
+++ b/Engine/Plugins/VirtualProduction/LevelSnapshots/Source/LevelSnapshotsEditor/Private/Data/LevelSnapshotsEditorData.cpp
@@ -36,6 +36,7 @@ ULevelSnapshotsEditorData::ULevelSnapshotsEditorData(const FObjectInitializer& O
 		ULevelSnapshotsFilterPreset* OldFilter = UserDefinedFilters;
 		UserDefinedFilters = NewFilterToEdit;
 		UserDefinedFilters->MarkTransactional();
+		UserDefinedFilters->OnFilterModified.AddUObject(this, &ULevelSnapshotsEditorData::HandleFilterChange);
 
 		FilterResults->Modify();
 		FilterResults->SetUserFilters(UserDefinedFilters);
