Make filters with short names easier to click

diff --git a/Engine/Plugins/VirtualProduction/LevelSnapshots/Source/LevelSnapshotsEditor/Private/Widgets/SLevelSnapshotsEditorFilter.cpp b/Engine/Plugins/VirtualProduction/LevelSnapshots/Source/LevelSnapshotsEditor/Private/Widgets/SLevelSnapshotsEditorFilter.cpp
--- a/Engine/Plugins/VirtualProduction/LevelSnapshots/Source/LevelSnapshotsEditor/Private/Widgets/SLevelSnapshotsEditorFilter.cpp
+++ b/Engine/Plugins/VirtualProduction/LevelSnapshots/Source/LevelSnapshotsEditor/Private/Widgets/SLevelSnapshotsEditorFilter.cpp
@@ -86,6 +86,7 @@ void SLevelSnapshotsEditorFilter::Construct(const FArguments& InArgs, const TWea
 				.Padding(FMargin(5.f, 1.f))
 				[
 					SAssignNew(FilterNamePtr, SClickableText)
+					.MinDesiredWidth(65.f)	//  SHoverableFilterActions (see below) makes clicking filters with short names difficult
 					.ColorAndOpacity(FLinearColor::White)
 					.Font(FEditorStyle::GetFontStyle("ContentBrowser.FilterNameFont"))
 					.ShadowOffset(FVector2D(1.f, 1.f))
