Only include struct subproperties if struct property itself is allowed.
Add bIncludeStructSubproperties to UPropertyHasNameFilter

diff --git a/Engine/Plugins/VirtualProduction/LevelSnapshots/Source/LevelShapshotFilters/Private/Builtin/PropertyHasNameFilter.cpp b/Engine/Plugins/VirtualProduction/LevelSnapshots/Source/LevelShapshotFilters/Private/Builtin/PropertyHasNameFilter.cpp
--- a/Engine/Plugins/VirtualProduction/LevelSnapshots/Source/LevelShapshotFilters/Private/Builtin/PropertyHasNameFilter.cpp
+++ b/Engine/Plugins/VirtualProduction/LevelSnapshots/Source/LevelShapshotFilters/Private/Builtin/PropertyHasNameFilter.cpp
@@ -3,32 +3,58 @@
 #include "PropertyHasNameFilter.h"
 
 EFilterResult::Type UPropertyHasNameFilter::IsPropertyValid(const FIsPropertyValidParams& Params) const
+{
+	if (GetPropertyMatchResult(Params.Property->GetName()) == EFilterResult::Include)
+	{
+		return EFilterResult::Include;
+	}
+	return bIncludeStructSubproperties && DoesAnyOwningPropertyMatch(Params)
+		? EFilterResult::Include : EFilterResult::Exclude;
+}
+
+bool UPropertyHasNameFilter::DoesAnyOwningPropertyMatch(const FIsPropertyValidParams& Params) const
+{
+	FIsPropertyValidParams ModifiableParams = Params;
+	for (int32 NameIndex = 0; NameIndex < Params.PropertyPath.Num(); ++NameIndex)
+	{
+		if (GetPropertyMatchResult(Params.PropertyPath[NameIndex]) == EFilterResult::Include)
+		{
+			return true;
+		}
+	}
+	return false;
+}
+
+EFilterResult::Type UPropertyHasNameFilter::GetPropertyMatchResult(const FString& PropertyNName) const
 {
 	class FNameMatcher
 	{
 	public:
 
 		FNameMatcher(ENameMatchingRule::Type MatchingRule, const TFieldPath<FProperty>& Property)
-            :
-			MatchingRule(MatchingRule),
-			PropertyName(Property->GetName())
+			: MatchingRule(MatchingRule)
+			, PropertyName(Property->GetName())
+		{}
+		FNameMatcher(ENameMatchingRule::Type MatchingRule, FString PropertyName)
+			: MatchingRule(MatchingRule)
+			, PropertyName(MoveTemp(PropertyName))
 		{}
 
 		bool MatchesName(const FString& AllowedName) const
 		{
 			switch(MatchingRule)
 			{
-				case ENameMatchingRule::MatchesExactly:
-					return MatchesExactly(AllowedName);
-				case ENameMatchingRule::MatchesIgnoreCase:
-					return MatchesIgnoreCase(AllowedName);
-				case ENameMatchingRule::ContainsExactly:
-					return ContainsExactly(AllowedName);
-				case ENameMatchingRule::ContainsIgnoreCase:
-					return ContainsIgnoreCase(AllowedName);
-				default:
-					ensure(false);
-					return false;
+			case ENameMatchingRule::MatchesExactly:
+				return MatchesExactly(AllowedName);
+			case ENameMatchingRule::MatchesIgnoreCase:
+				return MatchesIgnoreCase(AllowedName);
+			case ENameMatchingRule::ContainsExactly:
+				return ContainsExactly(AllowedName);
+			case ENameMatchingRule::ContainsIgnoreCase:
+				return ContainsIgnoreCase(AllowedName);
+			default:
+				ensure(false);
+				return false;
 			}
 		}
 	
@@ -57,8 +83,8 @@ EFilterResult::Type UPropertyHasNameFilter::IsPropertyValid(const FIsPropertyVal
 			return PropertyName.Find(AllowedName, ESearchCase::IgnoreCase) != INDEX_NONE;
 		}
 	};
-
-	const FNameMatcher NameMatcher(NameMatchingRule, Params.Property);
+	
+	const FNameMatcher NameMatcher(NameMatchingRule, PropertyNName);
 	for (const FString& AllowedPropertyName : AllowedNames)
 	{
 		if (NameMatcher.MatchesName(AllowedPropertyName))
diff --git a/Engine/Plugins/VirtualProduction/LevelSnapshots/Source/LevelShapshotFilters/Public/Builtin/PropertyHasNameFilter.h b/Engine/Plugins/VirtualProduction/LevelSnapshots/Source/LevelShapshotFilters/Public/Builtin/PropertyHasNameFilter.h
--- a/Engine/Plugins/VirtualProduction/LevelSnapshots/Source/LevelShapshotFilters/Public/Builtin/PropertyHasNameFilter.h
+++ b/Engine/Plugins/VirtualProduction/LevelSnapshots/Source/LevelShapshotFilters/Public/Builtin/PropertyHasNameFilter.h
@@ -38,10 +38,22 @@ public:
 
 private:
 
+	bool DoesAnyOwningPropertyMatch(const FIsPropertyValidParams& Params) const;
+	EFilterResult::Type GetPropertyMatchResult(const FString& PropertyNName) const;
+	
 	/* How to compare the property name to AllowedNames */
 	UPROPERTY(EditAnywhere, Category = "Config")
 	TEnumAsByte<ENameMatchingRule::Type> NameMatchingRule = ENameMatchingRule::ContainsIgnoreCase;
 
+	/**
+	 * Whether to implicitly include sub-properties of structs.
+	 *
+	 * Example: AllowedNames contains RelativeLocation.
+	 * If bIncludeStructSubproperties == true, then the properties X, Y, and Z are included. Otherwise, you must explicitly add them to AllowedNames.
+	 */
+	UPROPERTY(EditAnywhere, Category = "Config")
+	bool bIncludeStructSubproperties = true;
+
 	/* The names to match the property name against. */
 	UPROPERTY(EditAnywhere, Category = "Config")
 	TSet<FString> AllowedNames;
diff --git a/Engine/Plugins/VirtualProduction/LevelSnapshots/Source/LevelSnapshots/Private/Selection/ApplySnapshotFilter.cpp b/Engine/Plugins/VirtualProduction/LevelSnapshots/Source/LevelSnapshots/Private/Selection/ApplySnapshotFilter.cpp
--- a/Engine/Plugins/VirtualProduction/LevelSnapshots/Source/LevelSnapshots/Private/Selection/ApplySnapshotFilter.cpp
+++ b/Engine/Plugins/VirtualProduction/LevelSnapshots/Source/LevelSnapshots/Private/Selection/ApplySnapshotFilter.cpp
@@ -631,21 +631,15 @@ UE::LevelSnapshots::Private::FApplySnapshotFilter::EPropertySearchResult UE::Lev
 	check(ContainerContext.WorldContainer);
 	check(ContainerContext.SnapshotContainer); 
 
-	if (!bSkipEqualityTest
-		&& !bAllowUnchangedProperties
-		&& ArePropertyValuesIdentical(ContainerContext, PropertyInCommon))
-	{
-		return EPropertySearchResult::NoPropertiesFound;
-	}
-
-	// Now we can ask user whether they care about the property
 	TArray<FString> PropertyPath = ContainerContext.AuthoredPathInformation;
 	PropertyPath.Add(PropertyInCommon->GetAuthoredName());
 	const bool bIsPropertyValid = EFilterResult::CanInclude(Filter->IsPropertyValid(
         { DeserializedSnapshotActor, WorldActor, ContainerContext.SnapshotContainer, ContainerContext.WorldContainer, PropertyInCommon, PropertyPath }
         ));
-						
-	if (bIsPropertyValid)
+
+	if (bIsPropertyValid
+		&& (bSkipEqualityTest
+			|| TrackChangedProperties(ContainerContext, PropertyInCommon) == EPropertySearchResult::FoundProperties))
 	{
 		ContainerContext.SelectionToAddTo.AddProperty(ContainerContext.PropertyChain.MakeAppended(PropertyInCommon));
 		return EPropertySearchResult::FoundProperties;
@@ -653,19 +647,24 @@ UE::LevelSnapshots::Private::FApplySnapshotFilter::EPropertySearchResult UE::Lev
 	return EPropertySearchResult::NoPropertiesFound;
 }
 
-bool UE::LevelSnapshots::Private::FApplySnapshotFilter::ArePropertyValuesIdentical(FPropertyContainerContext& ContainerContext, FProperty* PropertyInCommon)
+UE::LevelSnapshots::Private::FApplySnapshotFilter::EPropertySearchResult UE::LevelSnapshots::Private::FApplySnapshotFilter::TrackChangedProperties(FPropertyContainerContext& ContainerContext, FProperty* PropertyInCommon)
 {
 	const TOptional<EPropertySearchResult> StructResult = HandlePossibleStructProperties(ContainerContext, PropertyInCommon);
 	if (StructResult.IsSet())
 	{
-		return *StructResult == EPropertySearchResult::NoPropertiesFound;
+		return *StructResult;
 	}
 	
 	const TOptional<EPropertySearchResult> SubobjectResult = HandlePossibleSubobjectProperties(ContainerContext, PropertyInCommon);
 	if (SubobjectResult.IsSet())
 	{
-		return *SubobjectResult == EPropertySearchResult::NoPropertiesFound;
+		return *SubobjectResult;
 	}
 
-	return AreSnapshotAndOriginalPropertiesEquivalent(Snapshot, PropertyInCommon, ContainerContext.SnapshotContainer, ContainerContext.WorldContainer, DeserializedSnapshotActor, WorldActor);
+	if (bAllowUnchangedProperties || !AreSnapshotAndOriginalPropertiesEquivalent(Snapshot, PropertyInCommon, ContainerContext.SnapshotContainer, ContainerContext.WorldContainer, DeserializedSnapshotActor, WorldActor))
+	{
+		return EPropertySearchResult::FoundProperties;
+	}
+	
+	return EPropertySearchResult::NoPropertiesFound;
 }
diff --git a/Engine/Plugins/VirtualProduction/LevelSnapshots/Source/LevelSnapshots/Private/Selection/ApplySnapshotFilter.h b/Engine/Plugins/VirtualProduction/LevelSnapshots/Source/LevelSnapshots/Private/Selection/ApplySnapshotFilter.h
--- a/Engine/Plugins/VirtualProduction/LevelSnapshots/Source/LevelSnapshots/Private/Selection/ApplySnapshotFilter.h
+++ b/Engine/Plugins/VirtualProduction/LevelSnapshots/Source/LevelSnapshots/Private/Selection/ApplySnapshotFilter.h
@@ -98,7 +98,7 @@ namespace UE::LevelSnapshots::Private
 		
 
 		EPropertySearchResult AnalyseProperty(FPropertyContainerContext& ContainerContext, FProperty* PropertyInCommon, bool bSkipEqualityTest = false);
-		bool ArePropertyValuesIdentical(FPropertyContainerContext& ContainerContext, FProperty* PropertyInCommon);
+		EPropertySearchResult TrackChangedProperties(FPropertyContainerContext& ContainerContext, FProperty* PropertyInCommon);
 
 		
 		ULevelSnapshot* Snapshot;
