Optimize RerunConstructionScript calls

diff --git a/Engine/Plugins/VirtualProduction/LevelSnapshots/Source/LevelSnapshots/Private/Archive/ApplySnapshotToEditorArchive.cpp b/Engine/Plugins/VirtualProduction/LevelSnapshots/Source/LevelSnapshots/Private/Archive/ApplySnapshotToEditorArchive.cpp
--- a/Engine/Plugins/VirtualProduction/LevelSnapshots/Source/LevelSnapshots/Private/Archive/ApplySnapshotToEditorArchive.cpp
+++ b/Engine/Plugins/VirtualProduction/LevelSnapshots/Source/LevelSnapshots/Private/Archive/ApplySnapshotToEditorArchive.cpp
@@ -13,7 +13,6 @@
 #include "Components/ActorComponent.h"
 #include "Engine/World.h"
 #include "GameFramework/Actor.h"
-#include "Misc/ScopeExit.h"
 #include "Serialization/ObjectReader.h"
 #include "Serialization/ObjectWriter.h"
 #include "UObject/TextProperty.h"
@@ -159,15 +158,11 @@ void UE::LevelSnapshots::Private::FApplySnapshotToEditorArchive::ApplyToExisting
 	UE_LOG(LogLevelSnapshots, Verbose, TEXT("Applying to existing object %s (class %s)"), *InOriginalObject->GetPathName(), *InOriginalObject->GetClass()->GetPathName());
 	const FApplySnapshotPropertiesScope NotifySnapshotListeners({ InOriginalObject, InSelectionMapForResolvingSubobjects, Selection, true });
 #if WITH_EDITOR
-	InOriginalObject->PreEditChange(nullptr);
-	ON_SCOPE_EXIT
-	{
-		InOriginalObject->PostEditChange();
-	};
+	InOriginalObject->Modify();
 #endif
 	
 	// Step 1: Serialise  properties that were different from CDO at time of snapshotting and that are still different from CDO
-	UE::LevelSnapshots::Private::FApplySnapshotToEditorArchive ApplySavedData(InObjectData, InSharedData, InOriginalObject, InSelectionMapForResolvingSubobjects, Selection, Cache);
+	FApplySnapshotToEditorArchive ApplySavedData(InObjectData, InSharedData, InOriginalObject, InSelectionMapForResolvingSubobjects, Selection, Cache);
 	InOriginalObject->Serialize(ApplySavedData);
 	
 	// Step 2: Serialise any remaining properties that were not covered: properties that were equal to the CDO value when the snapshot was taken but now are different from the CDO.
@@ -189,7 +184,7 @@ void UE::LevelSnapshots::Private::FApplySnapshotToEditorArchive::ApplyToRecreate
 	
 	// Apply all properties that we saved into the target actor.
 	// We assume that InOriginalObject was already created with the snapshot CDO as template: we do not need Step 2 from ApplyToExistingWorldObject.
-	UE::LevelSnapshots::Private::FApplySnapshotToEditorArchive ApplySavedData(InObjectData, InSharedData, InOriginalObject, InSelectionMapForResolvingSubobjects, {}, Cache);
+	FApplySnapshotToEditorArchive ApplySavedData(InObjectData, InSharedData, InOriginalObject, InSelectionMapForResolvingSubobjects, {}, Cache);
 	InOriginalObject->Serialize(ApplySavedData);
 }
 
diff --git a/Engine/Plugins/VirtualProduction/LevelSnapshots/Source/LevelSnapshots/Private/Data/Util/Restoration/RestoreActorUtil.cpp b/Engine/Plugins/VirtualProduction/LevelSnapshots/Source/LevelSnapshots/Private/Data/Util/Restoration/RestoreActorUtil.cpp
--- a/Engine/Plugins/VirtualProduction/LevelSnapshots/Source/LevelSnapshots/Private/Data/Util/Restoration/RestoreActorUtil.cpp
+++ b/Engine/Plugins/VirtualProduction/LevelSnapshots/Source/LevelSnapshots/Private/Data/Util/Restoration/RestoreActorUtil.cpp
@@ -69,8 +69,7 @@ namespace UE::LevelSnapshots::Private::Internal::Restore
 		UE_LOG(LogLevelSnapshots, Verbose, TEXT("ActorLabel is \"%s\" for \"%s\""), *OriginalActor->GetActorLabel(), *OriginalActor->GetPathName());
 #endif
 
-		TInlineComponentArray<UActorComponent*> DeserializedComponents;
-		Deserialized.GetValue()->GetComponents(DeserializedComponents);
+		TInlineComponentArray<UActorComponent*> DeserializedComponents(Deserialized.GetValue());
 		Internal::Restore::DeserializeComponents(OriginalActor, ActorData, WorldData,
 			[&SerializeComponent, &DeserializedComponents](
 				FSubobjectSnapshotData& SerializedCompData,
@@ -87,20 +86,23 @@ namespace UE::LevelSnapshots::Private::Internal::Restore
 		        if (DeserializedCompCounterpart)
 		        {
 	        		SerializeComponent(SerializedCompData, Comp, *DeserializedCompCounterpart);
-				
-					// We may have modified render information, e.g. for lights we may have changed intensity or colour
-					// It may be more efficient to track whether we actually changed render state
-					Comp->ReregisterComponent();
 		        }
 		    }
 		);
 
-		// Fixes up restored attachments
-		OriginalActor->UpdateComponentTransforms();
+		// Handle case in which AttachParentAfterRestore is restored after OriginalActor otherwise ReregisterAllComponents will detach OriginalActor again
+		const AActor* AttachParentAfterRestore = OriginalActor->GetAttachParentActor();
+		if (AttachParentAfterRestore && AttachParentAfterRestore->IsAttachedTo(OriginalActor))
+		{
+			AttachParentAfterRestore->GetRootComponent()->DetachFromComponent(FDetachmentTransformRules::KeepRelativeTransform);
+		}
+
+		OriginalActor->RerunConstructionScripts();
+		// Update component state, e.g. render state if intensity for lights was changed. Also avoids us having to call PostEditChangeProperty on every property.
+		OriginalActor->ReregisterAllComponents();
 
 #if WITH_EDITOR
 		// Update World Outliner. Usually called by USceneComponent::AttachToComponent.
-		const AActor* AttachParentAfterRestore = OriginalActor->GetAttachParentActor();
 		const bool bAttachParentChanged = AttachParentBeforeRestore != AttachParentAfterRestore;
 		if (bAttachParentChanged)
 		{
