Fix unit test (and minor formatting)

diff --git a/Engine/Plugins/VirtualProduction/LevelSnapshots/Source/LevelSnapshots/Private/Tests/CustomObjectSerializationTests.cpp b/Engine/Plugins/VirtualProduction/LevelSnapshots/Source/LevelSnapshots/Private/Tests/CustomObjectSerializationTests.cpp
--- a/Engine/Plugins/VirtualProduction/LevelSnapshots/Source/LevelSnapshots/Private/Tests/CustomObjectSerializationTests.cpp
+++ b/Engine/Plugins/VirtualProduction/LevelSnapshots/Source/LevelSnapshots/Private/Tests/CustomObjectSerializationTests.cpp
@@ -785,6 +785,7 @@ namespace UE::LevelSnapshots::Private::Tests
 			.ModifyWorld([&](UWorld* World)
 			{
 				Stub->TestActor = World->SpawnActor<ASnapshotTestActor>();
+				Stub->TestActor->IntProperty = 21;
 				Stub->TestActor->AllocateNonReflectedSubobject();
 				Stub->TestActor->NonReflectedSubobject->FloatProperty = 42.f;
 				Stub->TestActor->NonReflectedSubobject->IntProperty = 21;
@@ -794,11 +795,13 @@ namespace UE::LevelSnapshots::Private::Tests
 			.ModifyWorld([&](UWorld* World)
 			{
 				Stub->TestActor->DestroyNonReflectedSubobject();
+				// Needed otherwise no changes will be detected by hash
+				Stub->TestActor->IntProperty = 42;
 			})
 			.FilterProperties(Stub->TestActor, [&](const FPropertySelectionMap& SelectionMap)
 			{
 				const UE::LevelSnapshots::FRestorableObjectSelection ObjectSelection = SelectionMap.GetObjectSelection(Stub->TestActor);
-				TestTrue(TEXT("No changed actor properties"), ObjectSelection.GetPropertySelection() && ObjectSelection.GetPropertySelection()->GetSelectedProperties().Num() == 0 && ObjectSelection.GetPropertySelection()->HasCustomSerializedSubobjects());
+				TestTrue(TEXT("No changed actor properties"), ObjectSelection.GetPropertySelection() && ObjectSelection.GetPropertySelection()->GetSelectedProperties().Num() == 1 && ObjectSelection.GetPropertySelection()->HasCustomSerializedSubobjects());
 				TestTrue(TEXT("No component selection"), ObjectSelection.GetComponentSelection() == nullptr);
 				TestTrue(TEXT("Needs to restore custom subobject"), ObjectSelection.GetCustomSubobjectSelection() && ObjectSelection.GetCustomSubobjectSelection()->CustomSnapshotSubobjectsToRestore.Num() == 1);
 			})
