Fix override materials on nDisplay meshes showing as changed

diff --git a/Engine/Plugins/VirtualProduction/LevelSnapshots/Source/nDisplaySupport/Private/Material/DisplayMaterialOverrideFix.cpp b/Engine/Plugins/VirtualProduction/LevelSnapshots/Source/nDisplaySupport/Private/Material/DisplayMaterialOverrideFix.cpp
new file mode 100644
--- /dev/null
+++ b/Engine/Plugins/VirtualProduction/LevelSnapshots/Source/nDisplaySupport/Private/Material/DisplayMaterialOverrideFix.cpp
@@ -0,0 +1,39 @@
+﻿// Copyright Epic Games, Inc. All Rights Reserved.
+
+#include "DisplayMaterialOverrideFix.h"
+
+#include "ILevelSnapshotsModule.h"
+#include "Params/PropertyComparisonParams.h"
+#include "UObject/UnrealType.h"
+
+namespace UE::LevelSnapshots::nDisplay::Private::Internal
+{
+	const FName ClusterPropertyName("Cluster");
+
+	static UClass* GetDisplayClusterPreviewComponentClass()
+	{
+		static const FSoftClassPath ClassPath("/Script/DisplayCluster.DisplayClusterScreenComponent");
+		return ClassPath.ResolveClass();
+	}
+}
+
+
+void UE::LevelSnapshots::nDisplay::Private::FDisplayMaterialOverrideFix::Register(ILevelSnapshotsModule& Module)
+{
+	TSharedRef<FDisplayMaterialOverrideFix> MaterialFix = MakeShared<FDisplayMaterialOverrideFix>();
+	MaterialFix->OverrideMaterials = UMeshComponent::StaticClass()->FindPropertyByName(GET_MEMBER_NAME_CHECKED(UMeshComponent, OverrideMaterials));
+	check(MaterialFix->OverrideMaterials);
+	
+	Module.RegisterPropertyComparer(Internal::GetDisplayClusterPreviewComponentClass(), MaterialFix);
+}
+
+UE::LevelSnapshots::IPropertyComparer::EPropertyComparison UE::LevelSnapshots::nDisplay::Private::FDisplayMaterialOverrideFix::ShouldConsiderPropertyEqual(const FPropertyComparisonParams& Params) const
+{
+	if (Params.LeafProperty == OverrideMaterials && Params.InspectedClass->IsChildOf(Internal::GetDisplayClusterPreviewComponentClass()))
+	{
+		// The material is set dynamically every tick, it will always diff and we do not want to restore it
+		return EPropertyComparison::TreatEqual;
+	}
+	
+	return EPropertyComparison::CheckNormally;
+}
diff --git a/Engine/Plugins/VirtualProduction/LevelSnapshots/Source/nDisplaySupport/Private/Material/DisplayMaterialOverrideFix.h b/Engine/Plugins/VirtualProduction/LevelSnapshots/Source/nDisplaySupport/Private/Material/DisplayMaterialOverrideFix.h
new file mode 100644
--- /dev/null
+++ b/Engine/Plugins/VirtualProduction/LevelSnapshots/Source/nDisplaySupport/Private/Material/DisplayMaterialOverrideFix.h
@@ -0,0 +1,29 @@
+﻿// Copyright Epic Games, Inc. All Rights Reserved.
+
+#pragma once
+
+#include "CoreMinimal.h"
+#include "Interfaces/IPropertyComparer.h"
+
+namespace UE::LevelSnapshots
+{
+	class ILevelSnapshotsModule;
+}
+
+namespace UE::LevelSnapshots::nDisplay::Private
+{
+	/** Fixes material instances in UDisplayClusterPreviewComponent showing up as changed */
+	class FDisplayMaterialOverrideFix : public IPropertyComparer
+	{
+		FProperty* OverrideMaterials{};
+
+	public:
+
+		static void Register(ILevelSnapshotsModule& Module);
+	
+		//~ Begin IPropertyComparer Interface
+		virtual EPropertyComparison ShouldConsiderPropertyEqual(const FPropertyComparisonParams& Params) const override;
+		//~ End IPropertyComparer Interface
+	};
+}
+
diff --git a/Engine/Plugins/VirtualProduction/LevelSnapshots/Source/nDisplaySupport/Private/nDisplaySupportModule.cpp b/Engine/Plugins/VirtualProduction/LevelSnapshots/Source/nDisplaySupport/Private/nDisplaySupportModule.cpp
--- a/Engine/Plugins/VirtualProduction/LevelSnapshots/Source/nDisplaySupport/Private/nDisplaySupportModule.cpp
+++ b/Engine/Plugins/VirtualProduction/LevelSnapshots/Source/nDisplaySupport/Private/nDisplaySupportModule.cpp
@@ -5,6 +5,7 @@
 #include "ILevelSnapshotsModule.h"
 #include "Map/DisplayClusterConfigurationClusterNodeSerializer.h"
 #include "Map/DisplayClusterConfigurationClusterSerializer.h"
+#include "Material/DisplayMaterialOverrideFix.h"
 #include "Reference/DisplayClusterConfigurationDataSerializer.h"
 #include "Reference/DisplayClusterRootActorSerializer.h"
 
@@ -22,6 +23,8 @@ void UE::LevelSnapshots::nDisplay::Private::FnDisplaySupportModule::StartupModul
 		FDisplayClusterConfigurationDataSerializer::Register(LevelSnapshotsModule);
 		FDisplayClusterConfigurationClusterSerializer::Register(LevelSnapshotsModule);
 		FDisplayClusterConfigurationClusterNodeSerializer::Register(LevelSnapshotsModule);
+
+		FDisplayMaterialOverrideFix::Register(LevelSnapshotsModule);
 	}
 }
 void UE::LevelSnapshots::nDisplay::Private::FnDisplaySupportModule::ShutdownModule()
