Fix compile error

diff --git a/Engine/Plugins/VirtualProduction/LevelSnapshots/Source/nDisplaySupport/Private/Material/DisplayMaterialOverrideFix.cpp b/Engine/Plugins/VirtualProduction/LevelSnapshots/Source/nDisplaySupport/Private/Material/DisplayMaterialOverrideFix.cpp
--- a/Engine/Plugins/VirtualProduction/LevelSnapshots/Source/nDisplaySupport/Private/Material/DisplayMaterialOverrideFix.cpp
+++ b/Engine/Plugins/VirtualProduction/LevelSnapshots/Source/nDisplaySupport/Private/Material/DisplayMaterialOverrideFix.cpp
@@ -8,8 +8,6 @@
 
 namespace UE::LevelSnapshots::nDisplay::Private::Internal
 {
-	const FName ClusterPropertyName("Cluster");
-
 	static UClass* GetDisplayClusterPreviewComponentClass()
 	{
 		static const FSoftClassPath ClassPath("/Script/DisplayCluster.DisplayClusterScreenComponent");
