Add missing include

diff --git a/Engine/Plugins/VirtualProduction/LevelSnapshots/Source/nDisplaySupport/Private/Material/DisplayMaterialOverrideFix.cpp b/Engine/Plugins/VirtualProduction/LevelSnapshots/Source/nDisplaySupport/Private/Material/DisplayMaterialOverrideFix.cpp
--- a/Engine/Plugins/VirtualProduction/LevelSnapshots/Source/nDisplaySupport/Private/Material/DisplayMaterialOverrideFix.cpp
+++ b/Engine/Plugins/VirtualProduction/LevelSnapshots/Source/nDisplaySupport/Private/Material/DisplayMaterialOverrideFix.cpp
@@ -4,7 +4,9 @@
 
 #include "ILevelSnapshotsModule.h"
 #include "Params/PropertyComparisonParams.h"
+
 #include "UObject/UnrealType.h"
+#include "Components/MeshComponent.h"
 
 namespace UE::LevelSnapshots::nDisplay::Private::Internal
 {
