Fix renaming references not working for level snapshots

diff --git a/Engine/Plugins/VirtualProduction/LevelSnapshots/Source/LevelSnapshots/Private/Data/WorldSnapshotData.cpp b/Engine/Plugins/VirtualProduction/LevelSnapshots/Source/LevelSnapshots/Private/Data/WorldSnapshotData.cpp
--- a/Engine/Plugins/VirtualProduction/LevelSnapshots/Source/LevelSnapshots/Private/Data/WorldSnapshotData.cpp
+++ b/Engine/Plugins/VirtualProduction/LevelSnapshots/Source/LevelSnapshots/Private/Data/WorldSnapshotData.cpp
@@ -26,22 +26,18 @@ namespace UE::LevelSnapshots::Private::Internal
 {
 	static void CollectActorReferences(FWorldSnapshotData& WorldData, FArchive& Ar)
 	{
-		for (auto ActorIt = WorldData.ActorData.CreateConstIterator(); ActorIt; ++ActorIt)
+		for (auto ActorIt = WorldData.ActorData.CreateIterator(); ActorIt; ++ActorIt)
 		{
-			FSoftObjectPath SavedActorPath = ActorIt->Key; 
-			Ar << SavedActorPath;
-
-			FSoftClassPath ActorClass = ActorIt->Value.ActorClass;
-			Ar << ActorClass;
+			Ar << ActorIt->Key;
+			Ar << ActorIt->Value.ActorClass;
 		}
 	}
 
 	static void CollectClassDefaultReferences(FWorldSnapshotData& WorldData, FArchive& Ar)
 	{
-		for (auto ClassDefaultIt = WorldData.ClassDefaults.CreateConstIterator(); ClassDefaultIt; ++ClassDefaultIt)
+		for (auto ClassDefaultIt = WorldData.ClassDefaults.CreateIterator(); ClassDefaultIt; ++ClassDefaultIt)
 		{
-			FSoftClassPath Class = ClassDefaultIt->Key;
-			Ar << Class;
+			Ar << ClassDefaultIt->Key;
 		}
 	}
 	
