Fix snapshots sometimes not saving

diff --git a/Engine/Plugins/VirtualProduction/LevelSnapshots/Source/LevelSnapshots/Private/Data/WorldSnapshotData.cpp b/Engine/Plugins/VirtualProduction/LevelSnapshots/Source/LevelSnapshots/Private/Data/WorldSnapshotData.cpp
--- a/Engine/Plugins/VirtualProduction/LevelSnapshots/Source/LevelSnapshots/Private/Data/WorldSnapshotData.cpp
+++ b/Engine/Plugins/VirtualProduction/LevelSnapshots/Source/LevelSnapshots/Private/Data/WorldSnapshotData.cpp
@@ -52,11 +52,11 @@ namespace UE::LevelSnapshots::Private::Internal
 		Ar << WorldData.SerializedNames;
 		Ar << WorldData.SnapshotVersionInfo.CustomVersions;
 
-		// They're required... these names are for some reason not discovered by default...
-		TArray<FName> RequiredNames = { FName("UInt16Property"), EName::DoubleProperty};
-		for (FName Name : RequiredNames)
+		// Serialize hardcoded property names
+		for (int32 i = 0; i < static_cast<int32>(EName::MaxHardcodedNameIndex); ++i)
 		{
-			Ar << Name;
+			FName HardcodedName(static_cast<EName>(i));
+			Ar << HardcodedName;
 		}
 
 		auto ProcessProperty = [&Ar](const FProperty* Property)
