Fix stale snapshot object pointer

diff --git a/Engine/Plugins/VirtualProduction/LevelSnapshots/Source/LevelSnapshots/Private/Data/Util/SnapshotObjectUtil.cpp b/Engine/Plugins/VirtualProduction/LevelSnapshots/Source/LevelSnapshots/Private/Data/Util/SnapshotObjectUtil.cpp
--- a/Engine/Plugins/VirtualProduction/LevelSnapshots/Source/LevelSnapshots/Private/Data/Util/SnapshotObjectUtil.cpp
+++ b/Engine/Plugins/VirtualProduction/LevelSnapshots/Source/LevelSnapshots/Private/Data/Util/SnapshotObjectUtil.cpp
@@ -43,7 +43,7 @@ namespace UE::LevelSnapshots::Private::Internal
 		{
 			return nullptr;
 		}
-		if (const FSubobjectSnapshotCache* SubobjectCache = Cache.SubobjectCache.Find(SnapshotPathToSubobject); SubobjectCache && SubobjectCache->SnapshotObject.IsValid())
+		if (const FSubobjectSnapshotCache* SubobjectCache = Cache.SubobjectCache.Find(SnapshotPathToSubobject); SubobjectCache && SubobjectCache->SnapshotObject)
 		{
 			return SubobjectCache->SnapshotObject.Get();
 		}
@@ -144,7 +144,7 @@ namespace UE::LevelSnapshots::Private::Internal
 			}
 		}
 
-		if (!ensureMsgf(SubobjectCache->SnapshotObject.IsValid(), TEXT("Subobject pointer should have been set. Investigate.")))
+		if (!ensureMsgf(SubobjectCache->SnapshotObject, TEXT("Subobject pointer should have been set. Investigate.")))
 		{
 			return ResolvedObject;
 		}
@@ -242,7 +242,7 @@ namespace UE::LevelSnapshots::Private::Internal
 		}
 
 		FSubobjectSnapshotCache& SubobjectCache = Cache.SubobjectCache.FindOrAdd(WorldData.SerializedObjectReferences[ObjectPathIndex]);
-		if (!SubobjectCache.SnapshotObject.IsValid())
+		if (!SubobjectCache.SnapshotObject)
 		{
 			SubobjectCache.SnapshotObject = ExistingObject;
 			FLoadSnapshotObjectArchive::ApplyToSnapshotWorldObject(*SubobjectData, WorldData, Cache, ExistingObject, ProcessObjectDependency, LocalisationNamespace);
diff --git a/Engine/Plugins/VirtualProduction/LevelSnapshots/Source/LevelSnapshots/Public/Data/SnapshotDataCache.h b/Engine/Plugins/VirtualProduction/LevelSnapshots/Source/LevelSnapshots/Public/Data/SnapshotDataCache.h
--- a/Engine/Plugins/VirtualProduction/LevelSnapshots/Source/LevelSnapshots/Public/Data/SnapshotDataCache.h
+++ b/Engine/Plugins/VirtualProduction/LevelSnapshots/Source/LevelSnapshots/Public/Data/SnapshotDataCache.h
@@ -37,7 +37,7 @@ struct FSubobjectSnapshotCache
 	
 	/** Allocated in snapshot world */
 	UPROPERTY()
-	TWeakObjectPtr<UObject> SnapshotObject;
+	TObjectPtr<UObject> SnapshotObject;
 
 	/** Allocated in editor world */
 	UPROPERTY()
