Fix failing ensure on non-editor builds

diff --git a/Engine/Plugins/VirtualProduction/LevelSnapshots/Source/LevelSnapshots/Private/LevelSnapshotsModule.cpp b/Engine/Plugins/VirtualProduction/LevelSnapshots/Source/LevelSnapshots/Private/LevelSnapshotsModule.cpp
--- a/Engine/Plugins/VirtualProduction/LevelSnapshots/Source/LevelSnapshots/Private/LevelSnapshotsModule.cpp
+++ b/Engine/Plugins/VirtualProduction/LevelSnapshots/Source/LevelSnapshots/Private/LevelSnapshotsModule.cpp
@@ -86,11 +86,13 @@ namespace UE::LevelSnapshots::Private::Internal
 
 	static void DisableIrrelevantActorProperties(UE::LevelSnapshots::Private::FLevelSnapshotsModule& Module)
 	{
+#if WITH_EDITORONLY_DATA
 		const FProperty* ActorGuid = AActor::StaticClass()->FindPropertyByName(FName("ActorGuid"));
 		if (ensure(ActorGuid))
 		{
 			Module.AddExplicitlyUnsupportedProperties({ ActorGuid });
 		}
+#endif
 	}
 }
 
