Recreating actors creates them with a different display name

diff --git a/Engine/Plugins/VirtualProduction/LevelSnapshots/Source/LevelSnapshots/Private/Data/Util/Restoration/RestoreActorUtil.cpp b/Engine/Plugins/VirtualProduction/LevelSnapshots/Source/LevelSnapshots/Private/Data/Util/Restoration/RestoreActorUtil.cpp
--- a/Engine/Plugins/VirtualProduction/LevelSnapshots/Source/LevelSnapshots/Private/Data/Util/Restoration/RestoreActorUtil.cpp
+++ b/Engine/Plugins/VirtualProduction/LevelSnapshots/Source/LevelSnapshots/Private/Data/Util/Restoration/RestoreActorUtil.cpp
@@ -174,4 +174,9 @@ void UE::LevelSnapshots::Private::RestoreIntoRecreatedEditorWorldActor(AActor* O
 		FApplySnapshotToEditorArchive::ApplyToRecreatedEditorWorldObject(SerializedCompData, WorldData, Cache, Original, SelectedProperties); 
 	};
 	Internal::Restore::DeserializeIntoEditorWorldActor(OriginalActor, ActorData, WorldData, Cache, InLocalisationSnapshotPackage, DeserializeActor, DeserializeComponent);
+
+#if WITH_EDITOR
+	// Otherwise actor will show up with internal object name, e.g. actor previously called Cube will be StaticMeshActor1
+	OriginalActor->SetActorLabel(ActorData.ActorLabel);
+#endif
 }
