Recreating light shows "bad light" icon

diff --git a/Engine/Plugins/VirtualProduction/LevelSnapshots/Source/LevelSnapshots/Private/Data/Util/Restoration/RestoreActorUtil.cpp b/Engine/Plugins/VirtualProduction/LevelSnapshots/Source/LevelSnapshots/Private/Data/Util/Restoration/RestoreActorUtil.cpp
--- a/Engine/Plugins/VirtualProduction/LevelSnapshots/Source/LevelSnapshots/Private/Data/Util/Restoration/RestoreActorUtil.cpp
+++ b/Engine/Plugins/VirtualProduction/LevelSnapshots/Source/LevelSnapshots/Private/Data/Util/Restoration/RestoreActorUtil.cpp
@@ -178,5 +178,7 @@ void UE::LevelSnapshots::Private::RestoreIntoRecreatedEditorWorldActor(AActor* O
 #if WITH_EDITOR
 	// Otherwise actor will show up with internal object name, e.g. actor previously called Cube will be StaticMeshActor1
 	OriginalActor->SetActorLabel(ActorData.ActorLabel);
+	// Recreated actors have invalid lightning cache... e.g. recreated point lights will show error image (S_LightError)
+	OriginalActor->InvalidateLightingCacheDetailed(false);
 #endif
 }
