Editor crash while restoring specific snapshots

diff --git a/Engine/Plugins/VirtualProduction/LevelSnapshots/Source/LevelSnapshots/Private/Data/Util/Restoration/WorldDataUtil.cpp b/Engine/Plugins/VirtualProduction/LevelSnapshots/Source/LevelSnapshots/Private/Data/Util/Restoration/WorldDataUtil.cpp
--- a/Engine/Plugins/VirtualProduction/LevelSnapshots/Source/LevelSnapshots/Private/Data/Util/Restoration/WorldDataUtil.cpp
+++ b/Engine/Plugins/VirtualProduction/LevelSnapshots/Source/LevelSnapshots/Private/Data/Util/Restoration/WorldDataUtil.cpp
@@ -320,17 +320,22 @@ namespace UE::LevelSnapshots::Private::Internal
 				
 				const int32 NameLength = SubObjectPath.Len() - LastDotIndex - 1;
 				const FString ActorName = SubObjectPath.Right(NameLength);
+				ULevel* OverrideLevel = OwningLevelWorld->PersistentLevel;
 
 				const FName ActorFName = *ActorName;
 				FActorSpawnParameters SpawnParameters;
 				SpawnParameters.Name = ActorFName;
+				SpawnParameters.OverrideLevel = OverrideLevel;
 				SpawnParameters.bNoFail = true;
 				SpawnParameters.Template = Cast<AActor>(UE::LevelSnapshots::Private::GetClassDefault(WorldData, Cache, ActorClass));
 				SpawnParameters.ObjectFlags = ActorSnapshot->SerializedActorData.GetObjectFlags();
+				
 				Module.OnPreCreateActor(OwningLevelWorld, ActorClass, SpawnParameters);
-
-				checkf(SpawnParameters.Name == ActorFName, TEXT("You cannot change the name of the object"));
+				checkf(SpawnParameters.Name == ActorFName, TEXT("You cannot override the actor's name"));
+				checkf(SpawnParameters.OverrideLevel == OverrideLevel, TEXT("You cannot override the actor's level"))
+				
 				SpawnParameters.Name = ActorFName;
+				SpawnParameters.OverrideLevel = OwningLevelWorld->PersistentLevel;
 				SpawnParameters.NameMode = FActorSpawnParameters::ESpawnActorNameMode::Required_ErrorAndReturnNull;
 				if (AActor* RecreatedActor = OwningLevelWorld->SpawnActor(ActorClass, nullptr, SpawnParameters))
 				{
