Fix for: //UE5/Release-5.0 - Warning in Cook GoldenPath Win64 - SnapshotObject is not initialized properly

#[fyi] Dominik.Peacock
#[fyi] Jason.Walter

diff --git a/Engine/Plugins/VirtualProduction/LevelSnapshots/Source/LevelSnapshots/Public/Data/SnapshotDataCache.h b/Engine/Plugins/VirtualProduction/LevelSnapshots/Source/LevelSnapshots/Public/Data/SnapshotDataCache.h
--- a/Engine/Plugins/VirtualProduction/LevelSnapshots/Source/LevelSnapshots/Public/Data/SnapshotDataCache.h
+++ b/Engine/Plugins/VirtualProduction/LevelSnapshots/Source/LevelSnapshots/Public/Data/SnapshotDataCache.h
@@ -37,11 +37,11 @@ struct FSubobjectSnapshotCache
 	
 	/** Allocated in snapshot world */
 	UPROPERTY()
-	TObjectPtr<UObject> SnapshotObject;
+	TObjectPtr<UObject> SnapshotObject = nullptr;
 
 	/** Allocated in editor world */
 	UPROPERTY()
-	TWeakObjectPtr<UObject> EditorObject;
+	TWeakObjectPtr<UObject> EditorObject = nullptr;
 };
 
 USTRUCT()
