Fix nDisplaySupport module for Level Snapshots failing to load

diff --git a/Engine/Plugins/VirtualProduction/LevelSnapshots/Source/nDisplaySupport/Private/nDisplaySupportModule.cpp b/Engine/Plugins/VirtualProduction/LevelSnapshots/Source/nDisplaySupport/Private/nDisplaySupportModule.cpp
--- a/Engine/Plugins/VirtualProduction/LevelSnapshots/Source/nDisplaySupport/Private/nDisplaySupportModule.cpp
+++ b/Engine/Plugins/VirtualProduction/LevelSnapshots/Source/nDisplaySupport/Private/nDisplaySupportModule.cpp
@@ -9,7 +9,7 @@
 #include "Reference/DisplayClusterConfigurationDataSerializer.h"
 #include "Reference/DisplayClusterRootActorSerializer.h"
 
-#define LOCTEXT_NAMESPACE "FnDisplaySupportForLevelSnapshotsModule"
+#define LOCTEXT_NAMESPACE "nDisplaySupport"
 
 void UE::LevelSnapshots::nDisplay::Private::FnDisplaySupportModule::StartupModule()
 {
@@ -32,4 +32,4 @@ void UE::LevelSnapshots::nDisplay::Private::FnDisplaySupportModule::ShutdownModu
 
 #undef LOCTEXT_NAMESPACE
 	
-IMPLEMENT_MODULE(UE::LevelSnapshots::nDisplay::Private::FnDisplaySupportModule, nDisplaySupportForLevelSnapshots)
\ No newline at end of file
+IMPLEMENT_MODULE(UE::LevelSnapshots::nDisplay::Private::FnDisplaySupportModule, nDisplaySupport)
\ No newline at end of file
diff --git a/Engine/Plugins/VirtualProduction/LevelSnapshots/Source/nDisplaySupport/nDisplaySupport.Build.cs b/Engine/Plugins/VirtualProduction/LevelSnapshots/Source/nDisplaySupport/nDisplaySupport.Build.cs
--- a/Engine/Plugins/VirtualProduction/LevelSnapshots/Source/nDisplaySupport/nDisplaySupport.Build.cs
+++ b/Engine/Plugins/VirtualProduction/LevelSnapshots/Source/nDisplaySupport/nDisplaySupport.Build.cs
@@ -7,7 +7,6 @@ public class nDisplaySupport : ModuleRules
 	public nDisplaySupport(ReadOnlyTargetRules Target) : base(Target)
 	{
 		PCHUsage = ModuleRules.PCHUsageMode.UseExplicitOrSharedPCHs;
-		ShortName = "nDisplaySnapshots";
 
 		PublicDependencyModuleNames.AddRange(
 			new string[]
