Ensure thrown when restoring Level Snapshots that include subobjects

diff --git a/Engine/Plugins/VirtualProduction/LevelSnapshots/Source/nDisplaySupport/Private/Reference/DisplayClusterRootActorSerializer.cpp b/Engine/Plugins/VirtualProduction/LevelSnapshots/Source/nDisplaySupport/Private/Reference/DisplayClusterRootActorSerializer.cpp
--- a/Engine/Plugins/VirtualProduction/LevelSnapshots/Source/nDisplaySupport/Private/Reference/DisplayClusterRootActorSerializer.cpp
+++ b/Engine/Plugins/VirtualProduction/LevelSnapshots/Source/nDisplaySupport/Private/Reference/DisplayClusterRootActorSerializer.cpp
@@ -4,7 +4,9 @@
 
 #include "ILevelSnapshotsModule.h"
 #include "Params/ObjectSnapshotSerializationData.h"
+#include "Params/PropertyComparisonParams.h"
 
+#include "Components/MeshComponent.h"
 #include "UObject/UnrealType.h"
 
 namespace UE::LevelSnapshots::nDisplay::Private::Internal
@@ -30,7 +32,35 @@ void UE::LevelSnapshots::nDisplay::Private::FDisplayClusterRootActorSerializer::
 void UE::LevelSnapshots::nDisplay::Private::FDisplayClusterRootActorSerializer::Register(ILevelSnapshotsModule& Module)
 {
 	MarkPropertiesAsExplicitlyUnsupported(Module);
-	Module.RegisterCustomObjectSerializer(GetSupportedClass(), MakeShared<FDisplayClusterRootActorSerializer>());
+
+	const TSharedRef<FDisplayClusterRootActorSerializer> ActorSerializer = MakeShared<FDisplayClusterRootActorSerializer>();
+	Module.RegisterCustomObjectSerializer(GetSupportedClass(), ActorSerializer);
+	Module.RegisterPropertyComparer(UMeshComponent::StaticClass(), ActorSerializer);
+}
+
+UE::LevelSnapshots::nDisplay::Private::FDisplayClusterRootActorSerializer::FDisplayClusterRootActorSerializer()
+	: OverrideMaterials(UMeshComponent::StaticClass()->FindPropertyByName(GET_MEMBER_NAME_CHECKED(UMeshComponent, OverrideMaterials)))
+{
+	ensure(OverrideMaterials);
+}
+
+UE::LevelSnapshots::IPropertyComparer::EPropertyComparison UE::LevelSnapshots::nDisplay::Private::FDisplayClusterRootActorSerializer::ShouldConsiderPropertyEqual(const FPropertyComparisonParams& Params) const
+{
+	if (Params.LeafProperty == OverrideMaterials && Params.SnapshotActor->IsA(GetSupportedClass()))
+	{
+		// Materials are managed by UDisplayClusterPreviewComponent - they cannot be restored
+		// 
+		// The snapshot version of the material can be null when:
+		//		1. ProjectionPolicy is Mesh and
+		//		2. There is a mesh component on the nDisplay config
+		//		
+		// UDisplayClusterPreviewComponent is not created in the snapshot world
+		// UDisplayClusterPreviewComponent is the outer for materials assigned to override materials
+		// Since the UDisplayClusterPreviewComponent is null, the materials will also be null since they're missing an outer
+		return EPropertyComparison::TreatEqual;
+	}
+
+	return EPropertyComparison::CheckNormally;
 }
 
 UObject* UE::LevelSnapshots::nDisplay::Private::FDisplayClusterRootActorSerializer::FindSubobject(UObject* Owner) const
diff --git a/Engine/Plugins/VirtualProduction/LevelSnapshots/Source/nDisplaySupport/Private/Reference/DisplayClusterRootActorSerializer.h b/Engine/Plugins/VirtualProduction/LevelSnapshots/Source/nDisplaySupport/Private/Reference/DisplayClusterRootActorSerializer.h
--- a/Engine/Plugins/VirtualProduction/LevelSnapshots/Source/nDisplaySupport/Private/Reference/DisplayClusterRootActorSerializer.h
+++ b/Engine/Plugins/VirtualProduction/LevelSnapshots/Source/nDisplaySupport/Private/Reference/DisplayClusterRootActorSerializer.h
@@ -4,6 +4,7 @@
 
 #include "CoreMinimal.h"
 #include "Helpers/ReferenceSubobjectSerializer.h"
+#include "Interfaces/IPropertyComparer.h"
 
 namespace UE::LevelSnapshots
 {
@@ -12,15 +13,26 @@ namespace UE::LevelSnapshots
 
 namespace UE::LevelSnapshots::nDisplay::Private
 {
-	class FDisplayClusterRootActorSerializer : public TReferenceSubobjectSerializer<FDisplayClusterRootActorSerializer>
+	class FDisplayClusterRootActorSerializer
+		:
+		public TReferenceSubobjectSerializer<FDisplayClusterRootActorSerializer>,
+		public IPropertyComparer
 	{
 		static UClass* GetSupportedClass();
 		static void MarkPropertiesAsExplicitlyUnsupported(ILevelSnapshotsModule& Module);
 
+		const FProperty* OverrideMaterials;
+
 	public:
 
 		static void Register(ILevelSnapshotsModule& Module);
 
+		FDisplayClusterRootActorSerializer();
+
+		//~ Begin IPropertyComparer Interface
+		virtual EPropertyComparison ShouldConsiderPropertyEqual(const FPropertyComparisonParams& Params) const override;
+		//~ End IPropertyComparer Interface
+
 		//~ Begin FDisplayClusterRootActorSerializer Interface
 		UObject* FindSubobject(UObject* Owner) const;
 		//~ End FDisplayClusterRootActorSerializer Interface
