Level Snapshots:

-Add custom thumbnail to snapshot asset picker that displays thumbnail
-Add 'Update Snapshot Data' and 'Capture Thumbnail' options to LS asset editor context menu
-Restore Level Snapshots Context Picker
-Fix for: Level Snapshots does not allow snapshots from other maps

#[fyi] Jason.Walter

diff --git a/Engine/Plugins/VirtualProduction/LevelSnapshots/Source/LevelSnapshotsEditor/Private/Views/Input/SLevelSnapshotsEditorContextPicker.cpp b/Engine/Plugins/VirtualProduction/LevelSnapshots/Source/LevelSnapshotsEditor/Private/Views/Input/SLevelSnapshotsEditorContextPicker.cpp
new file mode 100644
--- /dev/null
+++ b/Engine/Plugins/VirtualProduction/LevelSnapshots/Source/LevelSnapshotsEditor/Private/Views/Input/SLevelSnapshotsEditorContextPicker.cpp
@@ -0,0 +1,240 @@
+// Copyright Epic Games, Inc. All Rights Reserved.
+
+#include "Views/Input/SLevelSnapshotsEditorContextPicker.h"
+
+#include "Data/LevelSnapshotsEditorData.h"
+
+#include "AssetRegistryModule.h"
+#include "Editor.h"
+#include "EditorStyleSet.h"
+#include "Engine/World.h"
+#include "Framework/MultiBox/MultiBoxBuilder.h"
+#include "Misc/MessageDialog.h"
+#include "Styling/SlateIconFinder.h"
+#include "Widgets/SBoxPanel.h"
+#include "Widgets/Images/SImage.h"
+#include "Widgets/Input/SComboButton.h"
+
+#define LOCTEXT_NAMESPACE "LevelSnapshotsEditor"
+
+SLevelSnapshotsEditorContextPicker::~SLevelSnapshotsEditorContextPicker()
+{
+	FEditorDelegates::OnMapOpened.Remove(OnMapOpenedDelegateHandle);
+	OnSelectWorldContextEvent.Unbind();
+}
+
+void SLevelSnapshotsEditorContextPicker::Construct(const FArguments& InArgs, ULevelSnapshotsEditorData* InEditorData)
+{
+	check(InEditorData);
+	SelectedWorldPath = InArgs._SelectedWorldPath.Get();
+	OnSelectWorldContextEvent = InArgs._OnSelectWorldContext;
+	DataPtr = InEditorData;
+
+	OnMapOpenedDelegateHandle = FEditorDelegates::OnMapOpened.AddLambda([this](const FString& FileName, bool bAsTemplate)
+    {
+		SetSelectedWorld(
+			DataPtr->GetEditorWorld()
+			);
+		
+        if (PickerButtonTextBlock.IsValid())
+        {
+            PickerButtonTextBlock.Get()->SetText(GetWorldPickerMenuButtonText(
+                SelectedWorldPath, FName(SelectedWorldPath.GetAssetName()))
+                );
+        }
+    });
+	
+	ChildSlot
+    .Padding(0.0f)
+    [
+        SNew(SComboButton)
+        .ContentPadding(0)
+        .ForegroundColor(FSlateColor::UseForeground())
+        .ButtonStyle(FEditorStyle::Get(), "ToggleButton")
+        .OnGetMenuContent(this, &SLevelSnapshotsEditorContextPicker::BuildWorldPickerMenu)
+        .ToolTipText(LOCTEXT("WorldPickerButtonTooltip", "The world context whose Level Snapshots you want to view"))
+        .ButtonContent()
+        [
+            SNew(SHorizontalBox)
+
+            + SHorizontalBox::Slot()
+            .AutoWidth()
+            [
+                SNew(SImage)
+                .Image(FEditorStyle::GetBrush("SceneOutliner.World"))
+            ]
+
+            + SHorizontalBox::Slot()
+            .Padding(3.f, 0.f)
+            [
+                SAssignNew(PickerButtonTextBlock, STextBlock)
+                .Text(GetCurrentContextText())
+            ]
+        ]
+    ];
+
+	RegisterWorldPickerWithEditorDataClass();
+}
+
+UWorld* SLevelSnapshotsEditorContextPicker::GetSelectedWorld() const
+{
+	const FAssetRegistryModule& AssetRegistryModule = FModuleManager::LoadModuleChecked<FAssetRegistryModule>(TEXT("AssetRegistry"));
+
+	return Cast<UWorld>(AssetRegistryModule.Get().GetAssetByObjectPath(FName(SelectedWorldPath.GetAssetPathString())).GetAsset());
+}
+
+FSoftObjectPath SLevelSnapshotsEditorContextPicker::GetSelectedWorldSoftPath() const
+{
+	return SelectedWorldPath;
+}
+
+TSharedRef<SWidget> SLevelSnapshotsEditorContextPicker::BuildWorldPickerMenu()
+{
+	FMenuBuilder MenuBuilder(true, nullptr);
+
+	// Get all worlds
+	const FAssetRegistryModule& AssetRegistryModule = FModuleManager::LoadModuleChecked<FAssetRegistryModule>(TEXT("AssetRegistry"));
+	TArray<FAssetData> WorldAssets;
+	AssetRegistryModule.Get().GetAssetsByClass(UWorld::StaticClass()->GetFName(), WorldAssets);
+
+	MenuBuilder.BeginSection("Other Worlds", LOCTEXT("OtherWorldsHeader", "Other Worlds")); 
+	for (const FAssetData& Asset : WorldAssets)
+	{
+		{
+			if (Asset.IsValid())
+			{
+				MenuBuilder.AddMenuEntry(
+					GetWorldPickerMenuButtonText(Asset.ToSoftObjectPath(), Asset.AssetName),
+					LOCTEXT("World", "World"),
+					FSlateIcon(),
+					FUIAction(
+						FExecuteAction::CreateRaw(
+							this, &SLevelSnapshotsEditorContextPicker::OnSetWorldContextSelection, Asset),
+						FCanExecuteAction(),
+						FIsActionChecked::CreateRaw(
+							this, &SLevelSnapshotsEditorContextPicker::ShouldRadioButtonBeChecked, Asset.ToSoftObjectPath())
+					),
+					NAME_None,
+					EUserInterfaceActionType::RadioButton
+				);
+			}
+		}
+	}
+	MenuBuilder.EndSection();
+
+	return MenuBuilder.MakeWidget();
+}
+
+void SLevelSnapshotsEditorContextPicker::RegisterWorldPickerWithEditorDataClass()
+{
+	if (DataPtr.IsValid())
+	{
+		SelectedWorldPath = DataPtr->GetEditorWorld();
+	}
+}
+
+FText SLevelSnapshotsEditorContextPicker::GetWorldDescription(const UWorld* World)
+{
+	FText PostFix;
+	if (World->WorldType == EWorldType::PIE)
+	{
+		switch(World->GetNetMode())
+		{
+		case NM_Client:
+			PostFix = FText::Format(LOCTEXT("ClientPostfixFormat", " (Client {0})"), FText::AsNumber(World->GetOutermost()->PIEInstanceID - 1));
+			break;
+		case NM_DedicatedServer:
+		case NM_ListenServer:
+			PostFix = LOCTEXT("ServerPostfix", " (Server)");
+			break;
+		case NM_Standalone:
+			PostFix = GEditor->bIsSimulatingInEditor ? LOCTEXT("SimulateInEditorPostfix", " (Simulate)") : LOCTEXT("PlayInEditorPostfix", " (PIE)");
+			break;
+
+		default:
+			break;
+		}
+	}
+	else if (World->WorldType == EWorldType::Editor)
+	{
+		PostFix = LOCTEXT("EditorPostfix", " (Editor)");
+	}
+
+	return FText::Format(LOCTEXT("WorldFormat", "{0}{1}"), FText::FromString(World->GetFName().GetPlainNameString()), PostFix);
+}
+
+FText SLevelSnapshotsEditorContextPicker::GetWorldPickerMenuButtonText(const FSoftObjectPath& AssetPath,
+	const FName& AssetName) const
+{
+	if (const bool bDoesAssetPointToEditorWorld = AssetPath == FSoftObjectPath(DataPtr->GetEditorWorld()))
+	{
+		const FText EditorLabel = LOCTEXT("EditorLabel", "Editor");
+		const FText FormattedEditorLabel = FText::Format(INVTEXT(" ({0})"), EditorLabel);
+
+		return FText::Format(INVTEXT("{0}{1}"),
+			FText::FromName(AssetName), (bDoesAssetPointToEditorWorld ? FormattedEditorLabel : FText::GetEmpty()));
+	}
+	else
+	{
+		return FText::Format(INVTEXT("{0}"), FText::FromName(AssetName));
+	}
+}
+
+FText SLevelSnapshotsEditorContextPicker::GetCurrentContextText() const
+{
+	UObject* WorldObject = SelectedWorldPath.ResolveObject();
+	UWorld* CurrentWorld = Cast<UWorld>(WorldObject);
+	check(CurrentWorld);
+	return GetWorldDescription(CurrentWorld);
+}
+
+const FSlateBrush* SLevelSnapshotsEditorContextPicker::GetBorderBrush(FSoftObjectPath WorldPath) const
+{
+	UObject* WorldObject = WorldPath.ResolveObject();
+	UWorld* CurrentWorld = Cast<UWorld>(WorldObject);
+	check(CurrentWorld);
+
+	if (CurrentWorld->WorldType == EWorldType::PIE)
+	{
+		return GEditor->bIsSimulatingInEditor ? FEditorStyle::GetBrush("LevelViewport.StartingSimulateBorder") : FEditorStyle::GetBrush("LevelViewport.StartingPlayInEditorBorder");
+	}
+	else
+	{
+		return FEditorStyle::GetBrush("LevelViewport.NoViewportBorder");
+	}
+}
+
+void SLevelSnapshotsEditorContextPicker::OnSetWorldContextSelection(const FAssetData Asset)
+{
+	check(Asset.IsValid());
+
+	SetSelectedWorld(Asset.ToSoftObjectPath());
+	
+	if (PickerButtonTextBlock.IsValid())
+	{
+		PickerButtonTextBlock.Get()->SetText(GetWorldPickerMenuButtonText(SelectedWorldPath, Asset.AssetName));
+	}
+}
+
+void SLevelSnapshotsEditorContextPicker::SetSelectedWorld(const FSoftObjectPath& SelectedWorld)
+{
+	SelectedWorldPath = SelectedWorld;
+	OnSelectWorldContextEvent.ExecuteIfBound(SelectedWorld);
+	
+	if (DataPtr.IsValid())
+	{
+		if (SelectedWorld != DataPtr->GetEditorWorld())
+		{
+			DataPtr->SetActiveSnapshot(nullptr);
+		
+			FMessageDialog::Open(EAppMsgType::Ok, LOCTEXT("UnsupportedOperation", "Your selected world is currently not loaded. In order to filter, you'll first have to load the selected world. Filtering operations will continue running on your last selected world."));
+		}
+	}
+}
+
+bool SLevelSnapshotsEditorContextPicker::ShouldRadioButtonBeChecked(const FSoftObjectPath InWorldSoftPath) const
+{
+	return SelectedWorldPath.IsAsset() && !SelectedWorldPath.IsNull() && SelectedWorldPath == InWorldSoftPath;
+}
+
+#undef LOCTEXT_NAMESPACE
diff --git a/Engine/Plugins/VirtualProduction/LevelSnapshots/Source/LevelSnapshotsEditor/Private/Views/Input/SLevelSnapshotsEditorContextPicker.h b/Engine/Plugins/VirtualProduction/LevelSnapshots/Source/LevelSnapshotsEditor/Private/Views/Input/SLevelSnapshotsEditorContextPicker.h
new file mode 100644
--- /dev/null
+++ b/Engine/Plugins/VirtualProduction/LevelSnapshots/Source/LevelSnapshotsEditor/Private/Views/Input/SLevelSnapshotsEditorContextPicker.h
@@ -0,0 +1,64 @@
+// Copyright Epic Games, Inc. All Rights Reserved.
+
+#pragma once
+
+#include "Widgets/DeclarativeSyntaxSupport.h"
+#include "Widgets/SCompoundWidget.h"
+#include "Widgets/Text/STextBlock.h"
+
+class ULevelSnapshotsEditorData;
+class FLevelSnapshotsEditorInput;
+class SLevelSnapshotsEditorBrowser;
+class SLevelSnapshotsEditorContextPicker;
+class UWorld;
+
+struct FAssetData;
+struct FLevelSnapshotsEditorViewBuilder;
+
+class SLevelSnapshotsEditorContextPicker : public SCompoundWidget
+{
+public:
+	
+	DECLARE_DELEGATE_OneParam(FOnSelectWorldContext, FSoftObjectPath);
+
+	~SLevelSnapshotsEditorContextPicker();
+	
+	SLATE_BEGIN_ARGS(SLevelSnapshotsEditorContextPicker) {}
+	/** Attribute for retrieving the current context */
+	SLATE_ATTRIBUTE(FSoftObjectPath, SelectedWorldPath)
+	/** Called when the user explicitly chooses a new context world. */
+	SLATE_EVENT(FOnSelectWorldContext, OnSelectWorldContext)
+	SLATE_END_ARGS()
+
+	void Construct(const FArguments& InArgs, ULevelSnapshotsEditorData* InEditorData);
+
+	UWorld* GetSelectedWorld() const;
+	FSoftObjectPath GetSelectedWorldSoftPath() const;
+
+private:
+	
+	static FText GetWorldDescription(const UWorld* World);
+
+	TSharedRef<SWidget> BuildWorldPickerMenu();
+	void RegisterWorldPickerWithEditorDataClass();
+
+
+	FText GetWorldPickerMenuButtonText(const FSoftObjectPath& AssetPath, const FName& AssetName) const;
+	FText GetCurrentContextText() const;
+	const FSlateBrush* GetBorderBrush(FSoftObjectPath WorldPath) const;
+
+	void OnSetWorldContextSelection(const FAssetData Asset);
+	void SetSelectedWorld(const FSoftObjectPath& SelectedWorld);
+
+	bool ShouldRadioButtonBeChecked(const FSoftObjectPath InWorldSoftPath) const;
+
+	FOnSelectWorldContext OnSelectWorldContextEvent;
+
+	TSharedPtr<STextBlock> PickerButtonTextBlock;
+	TWeakObjectPtr<ULevelSnapshotsEditorData> DataPtr;
+
+	// The selected radio button's world ref's soft path (for comparison)
+	FSoftObjectPath SelectedWorldPath;
+
+	FDelegateHandle OnMapOpenedDelegateHandle;
+};
diff --git a/Engine/Plugins/VirtualProduction/LevelSnapshots/Source/LevelSnapshotsEditor/Private/Views/Input/SLevelSnapshotsEditorInput.cpp b/Engine/Plugins/VirtualProduction/LevelSnapshots/Source/LevelSnapshotsEditor/Private/Views/Input/SLevelSnapshotsEditorInput.cpp
--- a/Engine/Plugins/VirtualProduction/LevelSnapshots/Source/LevelSnapshotsEditor/Private/Views/Input/SLevelSnapshotsEditorInput.cpp
+++ b/Engine/Plugins/VirtualProduction/LevelSnapshots/Source/LevelSnapshotsEditor/Private/Views/Input/SLevelSnapshotsEditorInput.cpp
@@ -3,14 +3,14 @@
 #include "Views/Input/SLevelSnapshotsEditorInput.h"
 
 #include "Data/LevelSnapshotsEditorData.h"
+#include "LevelSnapshotsLog.h"
+#include "SLevelSnapshotsEditorContextPicker.h"
 #include "Widgets/SLevelSnapshotsEditorBrowser.h"
 
 #include "Editor.h"
 #include "Engine/World.h"
 #include "Widgets/SBoxPanel.h"
 
-#include "LevelSnapshotsLog.h"
-
 #define LOCTEXT_NAMESPACE "LevelSnapshotsEditor"
 
 SLevelSnapshotsEditorInput::~SLevelSnapshotsEditorInput()
@@ -32,6 +32,14 @@ void SLevelSnapshotsEditorInput::Construct(const FArguments& InArgs, ULevelSnaps
 	[
 		SAssignNew(EditorInputOuterVerticalBox, SVerticalBox)
 
+		+ SVerticalBox::Slot()
+		.AutoHeight()
+		[
+			SAssignNew(EditorContextPickerPtr, SLevelSnapshotsEditorContextPicker, EditorData.Get())
+			.SelectedWorldPath(EditorData->GetEditorWorld())
+			.OnSelectWorldContext_Raw(this, &SLevelSnapshotsEditorInput::OverrideWorld)
+		]
+		
 		+ SVerticalBox::Slot()
 		[
 			SAssignNew(EditorBrowserWidgetPtr, SLevelSnapshotsEditorBrowser, InEditorData)
@@ -51,7 +59,7 @@ void SLevelSnapshotsEditorInput::OverrideWorld(FSoftObjectPath InNewContextPath)
 	if (!ensure(InNewContextPath.IsValid()))
 	{
 		UE_LOG(LogLevelSnapshots, Error,
-			TEXT("SLevelSnapshotsEditorInput::OverrideWorld: Unable to rebuild Snapshot Browser; InNewContext or BuilderPtr are invalid."));
+			TEXT("%hs: Unable to rebuild Snapshot Browser; InNewContext or BuilderPtr are invalid."), __FUNCTION__);
 		return;
 	}
 	
diff --git a/Engine/Plugins/VirtualProduction/LevelSnapshots/Source/LevelSnapshotsEditor/Private/Widgets/SLevelSnapshotsEditorBrowser.cpp b/Engine/Plugins/VirtualProduction/LevelSnapshots/Source/LevelSnapshotsEditor/Private/Widgets/SLevelSnapshotsEditorBrowser.cpp
--- a/Engine/Plugins/VirtualProduction/LevelSnapshots/Source/LevelSnapshotsEditor/Private/Widgets/SLevelSnapshotsEditorBrowser.cpp
+++ b/Engine/Plugins/VirtualProduction/LevelSnapshots/Source/LevelSnapshotsEditor/Private/Widgets/SLevelSnapshotsEditorBrowser.cpp
@@ -4,22 +4,27 @@
 
 #include "LevelSnapshot.h"
 #include "Data/LevelSnapshotsEditorData.h"
+#include "LevelSnapshotsEditor/Private/LevelSnapshotsEditorModule.h"
 
-#include "AssetRegistry/AssetRegistryModule.h"
+#include "AssetRegistry/IAssetRegistry.h"
 #include "AssetThumbnail.h"
 #include "ContentBrowserDelegates.h"
 #include "ContentBrowserModule.h"
 #include "Editor/EditorStyle/Public/EditorStyleSet.h"
 #include "IContentBrowserSingleton.h"
+#include "LevelSnapshotsEditorFunctionLibrary.h"
+#include "LevelSnapshotsEditorStyle.h"
 #include "LevelSnapshotsLog.h"
 #include "Misc/ScopedSlowTask.h"
 #include "Modules/ModuleManager.h"
-#include "Slate/Public/Framework/MultiBox/MultiBoxBuilder.h"
 #include "Toolkits/GlobalEditorCommonCommands.h"
 #include "UObject/Object.h"
 #include "UObject/SoftObjectPath.h"
+#include "Widgets/Input/SSlider.h"
 #include "Widgets/Layout/SBox.h"
-#include "Widgets/SToolTip.h"
+#include "Widgets/Input/SSearchBox.h"
+#include "Widgets/Layout/SWidgetSwitcher.h"
+#include "Widgets/Views/SListView.h"
 
 #define LOCTEXT_NAMESPACE "LevelSnapshotsEditor"
 
@@ -30,7 +35,8 @@ void SLevelSnapshotsEditorBrowser::Construct(const FArguments& InArgs, ULevelSna
 
 	check(OwningWorldPathAttribute.IsSet());
 
-	FContentBrowserModule& ContentBrowserModule = FModuleManager::Get().LoadModuleChecked<FContentBrowserModule>(TEXT("ContentBrowser"));
+	const FContentBrowserModule& ContentBrowserModule =
+		FModuleManager::Get().LoadModuleChecked<FContentBrowserModule>(TEXT("ContentBrowser"));
 
 	FARFilter ARFilter;
 	ARFilter.ClassNames.Add(ULevelSnapshot::StaticClass()->GetFName());
@@ -51,17 +57,25 @@ void SLevelSnapshotsEditorBrowser::Construct(const FArguments& InArgs, ULevelSna
 	AssetPickerConfig.ThumbnailScale = 0.8f;
 	AssetPickerConfig.Filter = ARFilter;
 	AssetPickerConfig.SelectionMode = ESelectionMode::Single;
-	AssetPickerConfig.OnAssetDoubleClicked = FOnAssetSelected::CreateSP(this, &SLevelSnapshotsEditorBrowser::OnAssetDoubleClicked);
-	AssetPickerConfig.OnShouldFilterAsset = FOnShouldFilterAsset::CreateSP(this, &SLevelSnapshotsEditorBrowser::OnShouldFilterAsset);
-	AssetPickerConfig.OnGetAssetContextMenu = FOnGetAssetContextMenu::CreateSP(this, &SLevelSnapshotsEditorBrowser::OnGetAssetContextMenu);
-
+	AssetPickerConfig.OnAssetDoubleClicked =
+		FOnAssetSelected::CreateSP(this, &SLevelSnapshotsEditorBrowser::OnAssetDoubleClicked);
+	AssetPickerConfig.OnShouldFilterAsset =
+		FOnShouldFilterAsset::CreateSP(this, &SLevelSnapshotsEditorBrowser::OnShouldFilterAsset);
+	AssetPickerConfig.OnGetAssetContextMenu =
+		FOnGetAssetContextMenu::CreateSP(this, &SLevelSnapshotsEditorBrowser::OnGetAssetContextMenu);
 	AssetPickerConfig.OnGetCustomAssetToolTip =
 		FOnGetCustomAssetToolTip::CreateSP(this, &SLevelSnapshotsEditorBrowser::CreateCustomTooltip);
-	
+		
+
 	ChildSlot
-		[
-			ContentBrowserModule.Get().CreateAssetPicker(AssetPickerConfig)
-		];
+	[
+		ContentBrowserModule.Get().CreateAssetPicker(AssetPickerConfig)
+	];	
+}
+
+SLevelSnapshotsEditorBrowser::~SLevelSnapshotsEditorBrowser()
+{
+	
 }
 
 void SLevelSnapshotsEditorBrowser::SelectAsset(const FAssetData& InAssetData) const
@@ -81,28 +95,6 @@ void SLevelSnapshotsEditorBrowser::SelectAsset(const FAssetData& InAssetData) co
 	}
 }
 
-TSharedRef<SToolTip> SLevelSnapshotsEditorBrowser::CreateCustomTooltip(FAssetData& AssetData)
-{
-	constexpr uint32 ThumbnailSize = 256;
-	TSharedRef<FAssetThumbnail> AssetThumbnail = MakeShared<FAssetThumbnail>(AssetData, ThumbnailSize, ThumbnailSize, UThumbnailManager::Get().GetSharedThumbnailPool());
-	
-	FAssetThumbnailConfig AssetThumbnailConfig;
-	AssetThumbnailConfig.bAllowFadeIn = false;
-	AssetThumbnailConfig.bAllowRealTimeOnHovered = false;
-	AssetThumbnailConfig.bForceGenericThumbnail = false;
-			
-	return SNew(SToolTip)
-	[
-		SNew(SBox)
-		.WidthOverride(ThumbnailSize)
-		.HeightOverride(ThumbnailSize)
-		.VAlign(VAlign_Top)
-		[
-			AssetThumbnail->MakeThumbnailWidget(AssetThumbnailConfig)
-		]
-	];
-}
-
 TArray<FAssetViewCustomColumn> SLevelSnapshotsEditorBrowser::GetCustomColumns() const
 {
 	TArray<FAssetViewCustomColumn> ReturnValue;
@@ -207,6 +199,42 @@ TSharedPtr<SWidget> SLevelSnapshotsEditorBrowser::OnGetAssetContextMenu(const TA
 
 	MenuBuilder.BeginSection(TEXT("Asset"), NSLOCTEXT("ReferenceViewerSchema", "AssetSectionLabel", "Asset"));
 	{
+		MenuBuilder.AddMenuEntry(
+		LOCTEXT("AssetTypeActions_LevelSnapshot_UpdateSnapshotData", "Update Snapshot Data"),
+		LOCTEXT("AssetTypeActions_LevelSnapshot_UpdateSnapshotDataToolTip", "Record a snapshot of the current map to this snapshot asset and update the thumbnail. Equivalent to 'Take Snapshot'. Select only one Level Snapshot asset at a time."),
+		FSlateIcon(FLevelSnapshotsEditorStyle::GetStyleSetName(), "LevelSnapshots.ToolbarButton.Small"),
+		FUIAction(
+			FExecuteAction::CreateLambda([SelectedAsset] {
+				if (const TObjectPtr<ULevelSnapshot> LevelSnapshotAsset = Cast<ULevelSnapshot>(SelectedAsset))
+				{
+					if (UWorld* World = ULevelSnapshotsEditorData::GetEditorWorld())
+					{
+						LevelSnapshotAsset->SnapshotWorld(World);
+						ULevelSnapshotsEditorFunctionLibrary::GenerateThumbnailForSnapshotAsset(LevelSnapshotAsset);
+						LevelSnapshotAsset->MarkPackageDirty();
+					}
+				}
+			}),
+			FCanExecuteAction::CreateLambda([] () { return true; })
+			)
+		);
+
+		MenuBuilder.AddMenuEntry(
+			LOCTEXT("AssetTypeActions_LevelSnapshot_CaptureThumbnail", "Capture Thumbnail"),
+			LOCTEXT("AssetTypeActions_LevelSnapshot_CaptureThumbnailsToolTip", "Capture and update the thumbnail only for this asset."),
+			FSlateIcon(FLevelSnapshotsEditorStyle::GetStyleSetName(), "LevelSnapshots.ToolbarButton.Small"),
+			FUIAction(
+			FExecuteAction::CreateLambda([SelectedAsset] {
+				if (const TObjectPtr<ULevelSnapshot> LevelSnapshotAsset = Cast<ULevelSnapshot>(SelectedAsset))
+				{
+					ULevelSnapshotsEditorFunctionLibrary::GenerateThumbnailForSnapshotAsset(LevelSnapshotAsset);
+					LevelSnapshotAsset->MarkPackageDirty();
+				}
+			}),
+			FCanExecuteAction::CreateLambda([] () { return true; })
+			)
+		);
+		
 		MenuBuilder.AddMenuEntry(
 			LOCTEXT("Browse", "Browse to Asset"),
 			LOCTEXT("BrowseTooltip", "Browses to the associated asset and selects it in the most recently used Content Browser (summoning one if necessary)"),
@@ -239,11 +267,96 @@ TSharedPtr<SWidget> SLevelSnapshotsEditorBrowser::OnGetAssetContextMenu(const TA
 				}),
 				FCanExecuteAction::CreateLambda([] () { return true; })
 			)
-	);
+		);
 	}
 	MenuBuilder.EndSection();
 
 	return MenuBuilder.MakeWidget();
 }
 
+TSharedRef<SToolTip> SLevelSnapshotsEditorBrowser::CreateCustomTooltip(FAssetData& AssetData)
+{
+	constexpr uint32 MaxTooltipWidth = 1024;
+	constexpr uint32 ThumbnailSize = 256;
+	constexpr uint32 MaxDescriptionHeight = ThumbnailSize / 3 * 2;
+	
+	const TSharedRef<FAssetThumbnail> AssetThumbnail =
+		MakeShareable(
+			new FAssetThumbnail(
+				AssetData.GetAsset(),
+				ThumbnailSize,
+				ThumbnailSize,
+				UThumbnailManager::Get().GetSharedThumbnailPool()
+			)
+		);
+	
+	FAssetThumbnailConfig AssetThumbnailConfig;
+	AssetThumbnailConfig.bAllowFadeIn = false;
+	AssetThumbnailConfig.bAllowRealTimeOnHovered = false;
+	AssetThumbnailConfig.bForceGenericThumbnail = false;
+	AssetThumbnailConfig.ColorStripOrientation = EThumbnailColorStripOrientation::VerticalRightEdge;
+
+	FString OutDescription;
+	AssetData.GetTagValue("SnapshotDescription",OutDescription);
+
+	FString OutCaptureTime;
+	AssetData.GetTagValue("CaptureTime",OutCaptureTime);
+	
+	return SNew(SToolTip)
+	[
+		SNew(SBox)
+		.MaxDesiredWidth(MaxTooltipWidth)
+		[
+			SNew(SHorizontalBox)
+
+			+SHorizontalBox::Slot()
+			.AutoWidth()
+			.VAlign(VAlign_Center)
+			.HAlign(HAlign_Left)
+			[
+				SNew(SBox)
+				.HeightOverride(ThumbnailSize + 5) // Offset for asset color strip
+				.WidthOverride(ThumbnailSize)
+				[
+					AssetThumbnail->MakeThumbnailWidget(AssetThumbnailConfig)
+				]
+			]
+
+			+SHorizontalBox::Slot()
+			.MaxWidth(MaxTooltipWidth - ThumbnailSize)
+			.VAlign(VAlign_Center)
+			.Padding(FMargin(5.f))
+			[
+				SNew(SVerticalBox)
+
+				+SVerticalBox::Slot()
+				.AutoHeight()
+				[
+					SNew(STextBlock)
+					.Text(FText::FromName(AssetData.AssetName))
+					.Font(FCoreStyle::GetDefaultFontStyle("Bold", 10))
+					.OverflowPolicy(ETextOverflowPolicy::Ellipsis)
+				]
+
+				+SVerticalBox::Slot()
+				.AutoHeight()
+				[
+					SNew(STextBlock)
+					.Text(FText::FromString(OutCaptureTime))
+				]
+
+				+SVerticalBox::Slot()
+				.MaxHeight(MaxDescriptionHeight)
+				[
+					SNew(STextBlock)
+					.Text(FText::FromString(OutDescription))
+					.AutoWrapText(true)
+					.WrappingPolicy(ETextWrappingPolicy::DefaultWrapping)
+					.OverflowPolicy(ETextOverflowPolicy::Ellipsis)				
+				]
+			]
+		]
+	];
+}
+
 #undef LOCTEXT_NAMESPACE
diff --git a/Engine/Plugins/VirtualProduction/LevelSnapshots/Source/LevelSnapshotsEditor/Private/Widgets/SLevelSnapshotsEditorBrowser.h b/Engine/Plugins/VirtualProduction/LevelSnapshots/Source/LevelSnapshotsEditor/Private/Widgets/SLevelSnapshotsEditorBrowser.h
--- a/Engine/Plugins/VirtualProduction/LevelSnapshots/Source/LevelSnapshotsEditor/Private/Widgets/SLevelSnapshotsEditorBrowser.h
+++ b/Engine/Plugins/VirtualProduction/LevelSnapshots/Source/LevelSnapshotsEditor/Private/Widgets/SLevelSnapshotsEditorBrowser.h
@@ -7,12 +7,9 @@
 #include "Widgets/SCompoundWidget.h"
 
 class ULevelSnapshotsEditorData;
-class UWorld;
-class SToolTip;
 
 struct FAssetData;
 struct FAssetViewCustomColumn;
-struct FSnapshotEditorViewData;
 
 class SLevelSnapshotsEditorBrowser : public SCompoundWidget
 {
@@ -26,16 +23,18 @@ public:
 	SLATE_END_ARGS()
 
 	void Construct(const FArguments& InArgs, ULevelSnapshotsEditorData* InEditorData);
+	virtual ~SLevelSnapshotsEditorBrowser() override;
+
 	void SelectAsset(const FAssetData& InAssetData) const;
 
 private:
 	
-	TSharedRef<SToolTip> CreateCustomTooltip(FAssetData& AssetData);
 	TArray<FAssetViewCustomColumn> GetCustomColumns() const;
 	void OnAssetDoubleClicked(const FAssetData& InAssetData) const;
 	bool OnShouldFilterAsset(const FAssetData& InAssetData) const;
 	TSharedPtr<SWidget> OnGetAssetContextMenu(const TArray<FAssetData>& SelectedAssets);
-
+	TSharedRef<SToolTip> CreateCustomTooltip(FAssetData& AssetData);
+	
 	TAttribute<FSoftObjectPath> OwningWorldPathAttribute;
 
 	TWeakObjectPtr<ULevelSnapshotsEditorData> EditorData;
