Level Snapshots:

-Fix for: NonUnity Compile UnrealEditor Win64 - SLevelSnapshotsEditorBrowser.cpp Error: 'MenuBuilder' uses undefined class 'FMenuBuilder'

#[fyi] Jason.Walter

diff --git a/Engine/Plugins/VirtualProduction/LevelSnapshots/Source/LevelSnapshotsEditor/Private/Widgets/SLevelSnapshotsEditorBrowser.cpp b/Engine/Plugins/VirtualProduction/LevelSnapshots/Source/LevelSnapshotsEditor/Private/Widgets/SLevelSnapshotsEditorBrowser.cpp
--- a/Engine/Plugins/VirtualProduction/LevelSnapshots/Source/LevelSnapshotsEditor/Private/Widgets/SLevelSnapshotsEditorBrowser.cpp
+++ b/Engine/Plugins/VirtualProduction/LevelSnapshots/Source/LevelSnapshotsEditor/Private/Widgets/SLevelSnapshotsEditorBrowser.cpp
@@ -5,16 +5,17 @@
 #include "LevelSnapshot.h"
 #include "Data/LevelSnapshotsEditorData.h"
 #include "LevelSnapshotsEditor/Private/LevelSnapshotsEditorModule.h"
+#include "LevelSnapshotsEditorFunctionLibrary.h"
+#include "LevelSnapshotsEditorStyle.h"
+#include "LevelSnapshotsLog.h"
 
 #include "AssetRegistry/IAssetRegistry.h"
 #include "AssetThumbnail.h"
 #include "ContentBrowserDelegates.h"
 #include "ContentBrowserModule.h"
 #include "Editor/EditorStyle/Public/EditorStyleSet.h"
+#include "Framework/MultiBox/MultiBoxBuilder.h"
 #include "IContentBrowserSingleton.h"
-#include "LevelSnapshotsEditorFunctionLibrary.h"
-#include "LevelSnapshotsEditorStyle.h"
-#include "LevelSnapshotsLog.h"
 #include "Misc/ScopedSlowTask.h"
 #include "Modules/ModuleManager.h"
 #include "Toolkits/GlobalEditorCommonCommands.h"
