Rename UI text to be less confusing

diff --git a/Engine/Plugins/VirtualProduction/LevelSnapshots/Source/LevelSnapshotsEditor/Private/Widgets/Filter/SSaveAndLoadFilters.cpp b/Engine/Plugins/VirtualProduction/LevelSnapshots/Source/LevelSnapshotsEditor/Private/Widgets/Filter/SSaveAndLoadFilters.cpp
--- a/Engine/Plugins/VirtualProduction/LevelSnapshots/Source/LevelSnapshotsEditor/Private/Widgets/Filter/SSaveAndLoadFilters.cpp
+++ b/Engine/Plugins/VirtualProduction/LevelSnapshots/Source/LevelSnapshotsEditor/Private/Widgets/Filter/SSaveAndLoadFilters.cpp
@@ -108,7 +108,7 @@ void SSaveAndLoadFilters::Construct(const FArguments& InArgs, ULevelSnapshotsEdi
             .Padding(0, 1, 0, 0)
             [
               SNew(STextBlock)
-              .Text(LOCTEXT("SavedToolbarButton", "Load/Save Filter"))
+              .Text(LOCTEXT("SavedToolbarButton", "Save / Load Preset"))
             ]
           ]
     ];
