More logging and converted check into ensure

diff --git a/Engine/Plugins/VirtualProduction/LevelSnapshots/Source/FoliageSupport/Private/FoliageSupport/FoliageInfoData.cpp b/Engine/Plugins/VirtualProduction/LevelSnapshots/Source/FoliageSupport/Private/FoliageSupport/FoliageInfoData.cpp
--- a/Engine/Plugins/VirtualProduction/LevelSnapshots/Source/FoliageSupport/Private/FoliageSupport/FoliageInfoData.cpp
+++ b/Engine/Plugins/VirtualProduction/LevelSnapshots/Source/FoliageSupport/Private/FoliageSupport/FoliageInfoData.cpp
@@ -3,10 +3,33 @@
 #include "FoliageSupport/FoliageInfoData.h"
 
 #include "InstancedFoliageActor.h"
+#include "LevelSnapshotsLog.h"
 #include "Serialization/MemoryReader.h"
 #include "Serialization/MemoryWriter.h"
 #include "Serialization/ObjectAndNameAsStringProxyArchive.h"
 
+namespace UE::LevelSnapshots::Foliage::Private
+{
+	static FString AsName(EFoliageImplType ImplType)
+	{
+		switch(ImplType)
+		{
+		case EFoliageImplType::StaticMesh:
+			return FString("StaticMesh");
+		case EFoliageImplType::Actor:
+			return FString("Actor");
+		case EFoliageImplType::ISMActor:
+			return FString("ISMActor");
+			
+		case EFoliageImplType::Unknown: 
+			return FString("Unknown");
+		default:
+			ensureMsgf(false, TEXT("Update this switch case."));
+			return FString("NewType (Update FoliageInfoData.cpp)");
+		}
+	}
+}
+
 FArchive& UE::LevelSnapshots::Foliage::Private::FFoliageInfoData::SerializeInternal(FArchive& Ar)
 {
 	Ar << ImplType;
@@ -23,6 +46,8 @@ void UE::LevelSnapshots::Foliage::Private::FFoliageInfoData::Save(FFoliageInfo&
 	
 	ImplType = DataToReadFrom.Type;
 	ComponentName = DataToReadFrom.GetComponent() ? DataToReadFrom.GetComponent()->GetFName() : FName(NAME_None);
+	
+	UE_CLOG(DataToReadFrom.GetComponent(), LogLevelSnapshots, Warning, TEXT("Could not save component for ImplType %s"), *AsName(ImplType));
 	RootArchive << DataToReadFrom;
 }
 
diff --git a/Engine/Plugins/VirtualProduction/LevelSnapshots/Source/FoliageSupport/Private/FoliageSupport/InstancedFoliageActorData.cpp b/Engine/Plugins/VirtualProduction/LevelSnapshots/Source/FoliageSupport/Private/FoliageSupport/InstancedFoliageActorData.cpp
--- a/Engine/Plugins/VirtualProduction/LevelSnapshots/Source/FoliageSupport/Private/FoliageSupport/InstancedFoliageActorData.cpp
+++ b/Engine/Plugins/VirtualProduction/LevelSnapshots/Source/FoliageSupport/Private/FoliageSupport/InstancedFoliageActorData.cpp
@@ -130,11 +130,10 @@ namespace UE::LevelSnapshots::Foliage::Private::Internal
 			if (RestorationInfo.ShouldSerializeFoliageType(FoliageInfoData))
 			{
 				FFoliageInfo* FoliageInfo = FindOrAddFoliageInfo(FoliageType, FoliageActor);
-				if (FoliageInfo)
+				if (FoliageInfo
+					&& ensureAlwaysMsgf(FoliageInfoData.GetComponentName(), TEXT("ComponentName is supposed to have been saved. Investigate. Maybe you used an unsupported foliage type (only EFoliageImplType::StaticMesh is supported)?")))
 				{
-					checkf(FoliageInfoData.GetComponentName(), TEXT("Did you enable support for non static mesh foliage?"));
 					HandleExistingFoliageUsingRequiredComponent(FoliageActor, *FoliageInfoData.GetComponentName(), PreexistingComponentToFoliageType, FoliageType);
-
 					FoliageInfoData.ApplyTo(*FoliageInfo, VersionInfo);
 				}
 			}
@@ -165,9 +164,9 @@ namespace UE::LevelSnapshots::Foliage::Private::Internal
 			if (RestorationInfo.ShouldSerializeFoliageType(InstanceData))
 			{
 				FFoliageInfo* FoliageInfo = FindOrAddFoliageInfo(FoliageType, FoliageActor);
-				if (FoliageInfo)
+				if (FoliageInfo
+					&& ensureAlwaysMsgf(InstanceData.GetComponentName(), TEXT("ComponentName is supposed to have been saved. Investigate. Maybe you used an unsupported foliage type (only EFoliageImplType::StaticMesh is supported)?")))
 				{
-					checkf(InstanceData.GetComponentName(), TEXT("Did you enable support for non static mesh foliage?"));
 					HandleExistingFoliageUsingRequiredComponent(FoliageActor, *InstanceData.GetComponentName(), PreexistingComponentToFoliageType, FoliageType);
 
 					InstanceData.ApplyTo(FoliageType, VersionInfo);
