Convert "Save Async" button to console variable

diff --git a/Engine/Plugins/VirtualProduction/LevelSnapshots/Source/LevelSnapshotsEditor/Private/Util/TakeSnapshotUtil.cpp b/Engine/Plugins/VirtualProduction/LevelSnapshots/Source/LevelSnapshotsEditor/Private/Util/TakeSnapshotUtil.cpp
--- a/Engine/Plugins/VirtualProduction/LevelSnapshots/Source/LevelSnapshotsEditor/Private/Util/TakeSnapshotUtil.cpp
+++ b/Engine/Plugins/VirtualProduction/LevelSnapshots/Source/LevelSnapshotsEditor/Private/Util/TakeSnapshotUtil.cpp
@@ -27,9 +27,11 @@
 
 #define LOCTEXT_NAMESPACE "LevelSnapshotEditorLibrary"
 
-namespace
+namespace UE::LevelSnapshots::Editor
 {
-	void HandleFormReply(const FText& InDescription, bool bShouldUseOverrides, bool bSaveAsync)
+	TAutoConsoleVariable<bool> CVarSaveAsync(TEXT("LevelSnapshots.SnapshotFormUsesSaveAsync"), true, TEXT("Set whether the Take Snapshot Form, opened by pressing the toolbar icon, saves snapshots using the SAVE_Async flag."));
+	
+	static void HandleFormReply(const FText& InDescription, bool bShouldUseOverrides)
 	{
 		UWorld* World = ULevelSnapshotsEditorData::GetEditorWorld();
 		if (!ensure(World))
@@ -53,10 +55,11 @@ namespace
 
 		const FString& ValidatedName = FPaths::MakeValidFileName(NewSnapshotName.ToString());
 
+		const bool bSaveAsync = CVarSaveAsync.GetValueOnAnyThread();
 		SnapshotEditor::TakeLevelSnapshotAndSaveToDisk(World, ValidatedName, NewSnapshotDir.ToString(), InDescription.ToString(), false, bSaveAsync);
 	}
 	
-	void DestroySnapshot(ULevelSnapshot* SnapshotAsset)
+	static void DestroySnapshot(ULevelSnapshot* SnapshotAsset)
 	{
 		SnapshotAsset->ClearFlags(RF_Public | RF_Standalone);
 		SnapshotAsset->Rename(nullptr, GetTransientPackage());
@@ -71,14 +74,14 @@ void SnapshotEditor::TakeSnapshotWithOptionalForm()
 	if (ULevelSnapshotsEditorSettings::Get()->bUseCreationForm)
 	{
 		TSharedRef<SWidget> CreationForm = SLevelSnapshotsEditorCreationForm::MakeAndShowCreationWindow(
-			FCloseCreationFormDelegate::CreateLambda([](const FText& Description, bool bSaveAsync)
+			FCloseCreationFormDelegate::CreateLambda([](const FText& Description)
 			{
-				HandleFormReply(Description, true, bSaveAsync);
+				UE::LevelSnapshots::Editor::HandleFormReply(Description, true);
 			}));
 	}
 	else
 	{
-		HandleFormReply(FText::GetEmpty(), false, false);
+		UE::LevelSnapshots::Editor::HandleFormReply(FText::GetEmpty(), false);
 	}
 }
 
@@ -148,7 +151,7 @@ ULevelSnapshot* SnapshotEditor::TakeLevelSnapshotAndSaveToDisk(UWorld* World, co
         if (!bSuccessful)
         {
             // Package and snapshot need to be destroyed again
-            DestroySnapshot(SnapshotAsset);
+            UE::LevelSnapshots::Editor:: DestroySnapshot(SnapshotAsset);
             ObjectTools::DeleteObjectsUnchecked({ SavePackage });
             
             NotificationItem->SetText(
diff --git a/Engine/Plugins/VirtualProduction/LevelSnapshots/Source/LevelSnapshotsEditor/Private/Widgets/SLevelSnapshotsEditorCreationForm.cpp b/Engine/Plugins/VirtualProduction/LevelSnapshots/Source/LevelSnapshotsEditor/Private/Widgets/SLevelSnapshotsEditorCreationForm.cpp
--- a/Engine/Plugins/VirtualProduction/LevelSnapshots/Source/LevelSnapshotsEditor/Private/Widgets/SLevelSnapshotsEditorCreationForm.cpp
+++ b/Engine/Plugins/VirtualProduction/LevelSnapshots/Source/LevelSnapshotsEditor/Private/Widgets/SLevelSnapshotsEditorCreationForm.cpp
@@ -205,22 +205,6 @@ void SLevelSnapshotsEditorCreationForm::Construct(
 			.Padding(2.f, 5.f)
 			[
 				SNew(SHorizontalBox)
-
-				// Save Async checkbox
-				+SHorizontalBox::Slot()
-				.Padding(3.f, 0.f)
-				.HAlign(HAlign_Left)
-				[
-					SNew(SCheckBox)
-					.IsChecked_Lambda([this]() { return bSaveAsync ? ECheckBoxState::Checked : ECheckBoxState::Unchecked; })
-					.OnCheckStateChanged_Lambda([this](ECheckBoxState NewState) { bSaveAsync = NewState == ECheckBoxState::Checked; })
-					.ToolTipText(NSLOCTEXT("LevelSnapshots", "CreationForm_SaveAsync_Tooltip", "Enabling may speed up saving for large levels."))
-					[
-						SNew(STextBlock)
-						.Font(FCoreStyle::GetDefaultFontStyle("Regular", 10))
-						.Text(NSLOCTEXT("LevelSnapshots", "CreationForm_SaveAsync", "Save async"))
-					]
-				]
 				
 				// Create snapshot button
 				+SHorizontalBox::Slot()
@@ -323,6 +307,6 @@ void SLevelSnapshotsEditorCreationForm::OnWindowClosed(const TSharedRef<SWindow>
 
 	if (bWasCreateSnapshotPressed)
 	{
-		CallOnCloseDelegate.ExecuteIfBound(DescriptionText, bSaveAsync);
+		CallOnCloseDelegate.ExecuteIfBound(DescriptionText);
 	}
 }
diff --git a/Engine/Plugins/VirtualProduction/LevelSnapshots/Source/LevelSnapshotsEditor/Private/Widgets/SLevelSnapshotsEditorCreationForm.h b/Engine/Plugins/VirtualProduction/LevelSnapshots/Source/LevelSnapshotsEditor/Private/Widgets/SLevelSnapshotsEditorCreationForm.h
--- a/Engine/Plugins/VirtualProduction/LevelSnapshots/Source/LevelSnapshotsEditor/Private/Widgets/SLevelSnapshotsEditorCreationForm.h
+++ b/Engine/Plugins/VirtualProduction/LevelSnapshots/Source/LevelSnapshotsEditor/Private/Widgets/SLevelSnapshotsEditorCreationForm.h
@@ -5,7 +5,7 @@
 #include "CoreMinimal.h"
 #include "Dialogs/CustomDialog.h"
 
-DECLARE_DELEGATE_TwoParams(FCloseCreationFormDelegate, const FText& /* Description */, bool /* bSaveAsync */);
+DECLARE_DELEGATE_OneParam(FCloseCreationFormDelegate, const FText& /* Description */);
 
 class ULevelSnapshotsSettings;
 class ULevelSnapshotsEditorSettings;
@@ -52,7 +52,6 @@ private:
 	bool bWasCreateSnapshotPressed = false;
 
 	FText DescriptionText;
-	bool bSaveAsync = true;
 
 	FCloseCreationFormDelegate CallOnCloseDelegate;
 };
