- Fixing Snapshot Filter creation method not setting Outer correctly

diff --git a/Engine/Plugins/VirtualProduction/LevelSnapshots/Source/LevelShapshotFilters/Private/FilterBlueprintFunctionLibrary.cpp b/Engine/Plugins/VirtualProduction/LevelSnapshots/Source/LevelShapshotFilters/Private/FilterBlueprintFunctionLibrary.cpp
--- a/Engine/Plugins/VirtualProduction/LevelSnapshots/Source/LevelShapshotFilters/Private/FilterBlueprintFunctionLibrary.cpp
+++ b/Engine/Plugins/VirtualProduction/LevelSnapshots/Source/LevelShapshotFilters/Private/FilterBlueprintFunctionLibrary.cpp
@@ -7,7 +7,7 @@
 
 ULevelSnapshotFilter* UFilterBlueprintFunctionLibrary::CreateFilterByClass(const TSubclassOf<ULevelSnapshotFilter>& Class, FName Name, UObject* Outer)
 {
-	if (IsValid(Outer))
+	if (!IsValid(Outer))
 	{
 		Outer = GetTransientPackage();
 	}
