Fix restoring foliage not working

diff --git a/Engine/Plugins/VirtualProduction/LevelSnapshots/Source/FoliageSupport/Private/FoliageSupport/FoliageInfoData.cpp b/Engine/Plugins/VirtualProduction/LevelSnapshots/Source/FoliageSupport/Private/FoliageSupport/FoliageInfoData.cpp
--- a/Engine/Plugins/VirtualProduction/LevelSnapshots/Source/FoliageSupport/Private/FoliageSupport/FoliageInfoData.cpp
+++ b/Engine/Plugins/VirtualProduction/LevelSnapshots/Source/FoliageSupport/Private/FoliageSupport/FoliageInfoData.cpp
@@ -45,9 +45,11 @@ void UE::LevelSnapshots::Foliage::Private::FFoliageInfoData::Save(FFoliageInfo&
 	RootArchive.SetCustomVersions(VersionInfo);
 	
 	ImplType = DataToReadFrom.Type;
-	ComponentName = DataToReadFrom.GetComponent() ? DataToReadFrom.GetComponent()->GetFName() : FName(NAME_None);
 	
-	UE_CLOG(DataToReadFrom.GetComponent(), LogLevelSnapshots, Warning, TEXT("Could not save component for ImplType %s"), *AsName(ImplType));
+	const bool bIsComponentValid = DataToReadFrom.GetComponent() != nullptr; 
+	ComponentName = bIsComponentValid ? DataToReadFrom.GetComponent()->GetFName() : FName(NAME_None);
+	UE_CLOG(!bIsComponentValid, LogLevelSnapshots, Warning, TEXT("Could not save component for ImplType %s"), *AsName(ImplType));
+
 	RootArchive << DataToReadFrom;
 }
 
diff --git a/Engine/Plugins/VirtualProduction/LevelSnapshots/Source/LevelSnapshots/Private/Data/LevelSnapshot.cpp b/Engine/Plugins/VirtualProduction/LevelSnapshots/Source/LevelSnapshots/Private/Data/LevelSnapshot.cpp
--- a/Engine/Plugins/VirtualProduction/LevelSnapshots/Source/LevelSnapshots/Private/Data/LevelSnapshot.cpp
+++ b/Engine/Plugins/VirtualProduction/LevelSnapshots/Source/LevelSnapshots/Private/Data/LevelSnapshot.cpp
@@ -78,9 +78,10 @@ bool ULevelSnapshot::SnapshotWorld(UWorld* TargetWorld)
 	}
 	Module.OnPreTakeSnapshot().Broadcast({this});
 
-	EnsureWorldInitialised();
 	MapPath = TargetWorld;
 	CaptureTime = FDateTime::UtcNow();
+	// Existing snapshot world actors are now invalid - clear previous cached data
+	RecreateSnapshotWorld();
 	SerializedData = UE::LevelSnapshots::Private::SnapshotWorld(TargetWorld);
 
 	Module.OnPostTakeSnapshot().Broadcast({ this });
@@ -271,6 +272,7 @@ void ULevelSnapshot::DiffWorld(UWorld* World, FActorPathConsumer HandleMatchedAc
 				UE::LevelSnapshots::Private::Internal::ConditionBreakOnActor(DebugActorName, OriginalActorPath);
 				SCOPED_SNAPSHOT_CORE_TRACE(HandleMatchedActor)
 				
+				UE_LOG(LogLevelSnapshots, VeryVerbose, TEXT("Matched actor %s"), *OriginalActorPath.ToString());
 				HandleMatchedActor.Execute(OriginalActorPath);
 			}
 		});
@@ -305,6 +307,15 @@ FString ULevelSnapshot::GenerateDebugLogInfo() const
 	return FString::Printf(TEXT("CaptureTime: %s. SnapshotVersionInfo: %s. Current engine version: %s."), *CaptureTime.ToString(), *GetSerializedData().SnapshotVersionInfo.ToString(), *Current.ToString());
 }
 
+void ULevelSnapshot::RecreateSnapshotWorld()
+{
+	if (SnapshotContainerWorld)
+	{
+		DestroyWorld();
+	}
+	EnsureWorldInitialised();
+}
+
 void ULevelSnapshot::EnsureWorldInitialised()
 {
 	if (SnapshotContainerWorld == nullptr)
@@ -369,9 +380,7 @@ void ULevelSnapshot::DestroyWorld()
 
 void ULevelSnapshot::ClearCache()
 {
-	Cache.ActorCache.Reset();
-	Cache.SubobjectCache.Reset();
-	Cache.ClassDefaultCache.Reset();
+	Cache.Reset();
 	
 #if WITH_EDITOR
 	CachedDiffedActors.Reset();
diff --git a/Engine/Plugins/VirtualProduction/LevelSnapshots/Source/LevelSnapshots/Public/Data/LevelSnapshot.h b/Engine/Plugins/VirtualProduction/LevelSnapshots/Source/LevelSnapshots/Public/Data/LevelSnapshot.h
--- a/Engine/Plugins/VirtualProduction/LevelSnapshots/Source/LevelSnapshots/Public/Data/LevelSnapshot.h
+++ b/Engine/Plugins/VirtualProduction/LevelSnapshots/Source/LevelSnapshots/Public/Data/LevelSnapshot.h
@@ -101,6 +101,7 @@ private:
 
 	FString GenerateDebugLogInfo() const;
 
+	void RecreateSnapshotWorld();
 	void EnsureWorldInitialised();
 	void DestroyWorld();
 	void ClearCache();
diff --git a/Engine/Plugins/VirtualProduction/LevelSnapshots/Source/LevelSnapshots/Public/Data/SnapshotDataCache.h b/Engine/Plugins/VirtualProduction/LevelSnapshots/Source/LevelSnapshots/Public/Data/SnapshotDataCache.h
--- a/Engine/Plugins/VirtualProduction/LevelSnapshots/Source/LevelSnapshots/Public/Data/SnapshotDataCache.h
+++ b/Engine/Plugins/VirtualProduction/LevelSnapshots/Source/LevelSnapshots/Public/Data/SnapshotDataCache.h
@@ -67,4 +67,11 @@ struct FSnapshotDataCache
 
 	UPROPERTY()
 	TMap<FSoftClassPath, FClassDefaultSnapshotCache> ClassDefaultCache;
+
+	void Reset()
+	{
+		ActorCache.Reset();
+		SubobjectCache.Reset();
+		ClassDefaultCache.Reset();
+	}
 };
