Fix unit tests spawning actors sometimes failing

diff --git a/Engine/Plugins/VirtualProduction/LevelSnapshots/Source/LevelSnapshots/Private/Data/Util/Restoration/WorldDataUtil.cpp b/Engine/Plugins/VirtualProduction/LevelSnapshots/Source/LevelSnapshots/Private/Data/Util/Restoration/WorldDataUtil.cpp
--- a/Engine/Plugins/VirtualProduction/LevelSnapshots/Source/LevelSnapshots/Private/Data/Util/Restoration/WorldDataUtil.cpp
+++ b/Engine/Plugins/VirtualProduction/LevelSnapshots/Source/LevelSnapshots/Private/Data/Util/Restoration/WorldDataUtil.cpp
@@ -327,6 +327,7 @@ namespace UE::LevelSnapshots::Private::Internal
 				SpawnParameters.Name = ActorFName;
 				SpawnParameters.OverrideLevel = OverrideLevel;
 				SpawnParameters.bNoFail = true;
+				SpawnParameters.SpawnCollisionHandlingOverride = ESpawnActorCollisionHandlingMethod::AlwaysSpawn;
 				SpawnParameters.Template = Cast<AActor>(UE::LevelSnapshots::Private::GetClassDefault(WorldData, Cache, ActorClass));
 				SpawnParameters.ObjectFlags = ActorSnapshot->SerializedActorData.GetObjectFlags();
 				
diff --git a/Engine/Plugins/VirtualProduction/LevelSnapshots/Source/LevelSnapshots/Private/Tests/ApplyToWorldTests.cpp b/Engine/Plugins/VirtualProduction/LevelSnapshots/Source/LevelSnapshots/Private/Tests/ApplyToWorldTests.cpp
--- a/Engine/Plugins/VirtualProduction/LevelSnapshots/Source/LevelSnapshots/Private/Tests/ApplyToWorldTests.cpp
+++ b/Engine/Plugins/VirtualProduction/LevelSnapshots/Source/LevelSnapshots/Private/Tests/ApplyToWorldTests.cpp
@@ -165,7 +165,7 @@ namespace UE::LevelSnapshots::Private::Tests
 		FSnapshotTestRunner()
 			.ModifyWorld([&](UWorld* World)
 			{
-				MaterialAndMesh = World->SpawnActor<ASnapshotTestActor>();
+				MaterialAndMesh = ASnapshotTestActor::Spawn(World);
 			
 				MaterialAndMesh->InstancedMeshComponent->AddInstance(StartFirstTransform);
 				MaterialAndMesh->InstancedMeshComponent->AddInstance(StartSecondTransform);
@@ -203,7 +203,7 @@ namespace UE::LevelSnapshots::Private::Tests
 		FSnapshotTestRunner()
 			.ModifyWorld([&](UWorld* World)
 			{
-				TestActor = World->SpawnActor<ASnapshotTestActor>();
+				TestActor = ASnapshotTestActor::Spawn(World);
 			})
 			.TakeSnapshot()
 
@@ -233,7 +233,7 @@ namespace UE::LevelSnapshots::Private::Tests
 		FSnapshotTestRunner()
 			.ModifyWorld([&](UWorld* World)
 			{
-				Actor = World->SpawnActor<ASnapshotTestActor>();
+				Actor = ASnapshotTestActor::Spawn(World);
 				Actor->AllocateSubobjects();
 			})
 			.TakeSnapshot()
@@ -255,9 +255,9 @@ namespace UE::LevelSnapshots::Private::Tests
 		FSnapshotTestRunner()
 			.ModifyWorld([&](UWorld* World)
 			{
-				ParentOne = World->SpawnActor<ASnapshotTestActor>();
-				ParentTwo = World->SpawnActor<ASnapshotTestActor>();
-				Child = World->SpawnActor<ASnapshotTestActor>();
+				ParentOne = ASnapshotTestActor::Spawn(World, "ParentOne");
+				ParentTwo = ASnapshotTestActor::Spawn(World, "ParentTwo");
+				Child = ASnapshotTestActor::Spawn(World, "Child");
 
 				Child->AttachToActor(ParentOne, FAttachmentTransformRules::KeepRelativeTransform);
 			})
diff --git a/Engine/Plugins/VirtualProduction/LevelSnapshots/Source/LevelSnapshots/Private/Tests/ComponentTests.cpp b/Engine/Plugins/VirtualProduction/LevelSnapshots/Source/LevelSnapshots/Private/Tests/ComponentTests.cpp
--- a/Engine/Plugins/VirtualProduction/LevelSnapshots/Source/LevelSnapshots/Private/Tests/ComponentTests.cpp
+++ b/Engine/Plugins/VirtualProduction/LevelSnapshots/Source/LevelSnapshots/Private/Tests/ComponentTests.cpp
@@ -40,7 +40,7 @@ namespace UE::LevelSnapshots::Private::Tests
 		FSnapshotTestRunner()
 			.ModifyWorld([&](UWorld* World)
 			{
-				Actor = World->SpawnActor<ASnapshotTestActor>();
+				Actor = ASnapshotTestActor::Spawn(World);;
 
 				InstancedStaysOnActor = NewObject<UStaticMeshComponent>(Actor, UStaticMeshComponent::StaticClass(), "InstancedStaysOnActor");
 				Actor->AddInstanceComponent(InstancedStaysOnActor);
@@ -98,8 +98,8 @@ namespace UE::LevelSnapshots::Private::Tests
 		FSnapshotTestRunner()
 			.ModifyWorld([&](UWorld* World)
 			{
-				ActorWithAddedComponent = World->SpawnActor<AActor>();
-				ActorWithRemovedComponent = World->SpawnActor<AActor>();
+				ActorWithAddedComponent = ASnapshotTestActor::Spawn(World, "ActorWithAddedComponent");
+				ActorWithRemovedComponent = ASnapshotTestActor::Spawn(World, "ActorWithRemovedComponent");
 
 				AddedComponent = NewObject<UStaticMeshComponent>(ActorWithRemovedComponent, UStaticMeshComponent::StaticClass(), "AddedComponent");
 				ActorWithRemovedComponent->AddInstanceComponent(AddedComponent);
@@ -150,7 +150,7 @@ namespace UE::LevelSnapshots::Private::Tests
 		FSnapshotTestRunner()
 			.ModifyWorld([&](UWorld* World)
 			{
-				Actor = World->SpawnActor<ASnapshotTestActor>();
+				Actor = ASnapshotTestActor::Spawn(World);;
 			
 				OriginalMeshComp = Actor->InstancedMeshComponent;
 				OriginalLightComp = Actor->PointLightComponent;
@@ -187,7 +187,7 @@ namespace UE::LevelSnapshots::Private::Tests
 		FSnapshotTestRunner()
 			.ModifyWorld([&](UWorld* World)
 			{
-				Actor = World->SpawnActor<ASnapshotTestActor>();
+				Actor = ASnapshotTestActor::Spawn(World);;
 			})
 			.TakeSnapshot()
 			.ModifyWorld([&](UWorld* World)
@@ -225,7 +225,7 @@ namespace UE::LevelSnapshots::Private::Tests
 		FSnapshotTestRunner()
 			.ModifyWorld([&](UWorld* World)
 			{
-				Actor = World->SpawnActor<ASnapshotTestActor>();
+				Actor = ASnapshotTestActor::Spawn(World);;
 
 				// This component will stay on the actor
 				InstancedStaysOnActor = NewObject<UStaticMeshComponent>(Actor, UStaticMeshComponent::StaticClass(), "InstancedStaysOnActor");
diff --git a/Engine/Plugins/VirtualProduction/LevelSnapshots/Source/LevelSnapshots/Private/Tests/CustomObjectSerializationTests.cpp b/Engine/Plugins/VirtualProduction/LevelSnapshots/Source/LevelSnapshots/Private/Tests/CustomObjectSerializationTests.cpp
--- a/Engine/Plugins/VirtualProduction/LevelSnapshots/Source/LevelSnapshots/Private/Tests/CustomObjectSerializationTests.cpp
+++ b/Engine/Plugins/VirtualProduction/LevelSnapshots/Source/LevelSnapshots/Private/Tests/CustomObjectSerializationTests.cpp
@@ -167,7 +167,7 @@ namespace UE::LevelSnapshots::Private::Tests
 		FSnapshotTestRunner()
 			.ModifyWorld([&](UWorld* World)
 			{
-				Stub->TestActor = World->SpawnActor<ASnapshotTestActor>();
+				Stub->TestActor = ASnapshotTestActor::Spawn(World);;
 
 				Stub->TestActor->InstancedOnlySubobject_DefaultSubobject->FloatProperty = 42.f;
 				Stub->TestActor->InstancedOnlySubobject_DefaultSubobject->IntProperty = 21;
@@ -231,7 +231,6 @@ namespace UE::LevelSnapshots::Private::Tests
 			virtual void OnTakeSnapshot(UObject* EditorObject, ICustomSnapshotSerializationData& DataStorage) override
 			{
 				ASnapshotTestActor* Actor = Cast<ASnapshotTestActor>(EditorObject);
-			
 				const int32 Index = DataStorage.AddSubobjectSnapshot(Actor->NonReflectedSubobject);
 				DataStorage.WriteObjectAnnotation(FObjectAnnotator::CreateLambda([Actor](FArchive& Archive)
 				{
@@ -263,7 +262,6 @@ namespace UE::LevelSnapshots::Private::Tests
 		
 			virtual UObject* FindSubobjectInEditorWorld(UObject* EditorObject, const ISnapshotSubobjectMetaData& ObjectData, const ICustomSnapshotSerializationData& DataStorage) override
 			{
-			
 				ASnapshotTestActor* Actor = Cast<ASnapshotTestActor>(EditorObject);
 				return Actor->NonReflectedSubobject;
 			}
@@ -271,7 +269,6 @@ namespace UE::LevelSnapshots::Private::Tests
 			virtual void PostApplySnapshotProperties(UObject* OriginalObject, const ICustomSnapshotSerializationData& DataStorage) override
 			{
 				ASnapshotTestActor* Actor = Cast<ASnapshotTestActor>(OriginalObject);
-			
 				DataStorage.ReadObjectAnnotation(FObjectAnnotator::CreateLambda([Actor](FArchive& Archive)
 				{
 					Archive << Actor->NonReflectedName;
@@ -300,16 +297,18 @@ namespace UE::LevelSnapshots::Private::Tests
 		};
 
 		ASnapshotTestActor* TestActor = nullptr;
+		// Do not use ASnapshotTestActor as class for referenced actors because we don't want them to be processed by FStub
 		AStaticMeshActor* FirstReferencedActor = nullptr;
 		AStaticMeshActor* SecondReferencedActor = nullptr;
 	
 		FSnapshotTestRunner()
 			.ModifyWorld([&](UWorld* World)
 			{
-				TestActor = World->SpawnActor<ASnapshotTestActor>();
+				TestActor = ASnapshotTestActor::Spawn(World, "TestActor");
 				TestActor->AllocateNonReflectedSubobject();
-				FirstReferencedActor = World->SpawnActor<AStaticMeshActor>();
-				SecondReferencedActor = World->SpawnActor<AStaticMeshActor>();
+				// Spawn at different locations to avoid them failing to spawn
+				FirstReferencedActor = World->SpawnActor<AStaticMeshActor>(FVector(500.f), FRotator());
+				SecondReferencedActor = World->SpawnActor<AStaticMeshActor>(FVector(-500.f), FRotator());
 			
 				TestActor->IntProperty = 1;
 				TestActor->NonReflectedName = FName("TestNonReflectedName_OnActor");
@@ -444,7 +443,7 @@ namespace UE::LevelSnapshots::Private::Tests
 		FSnapshotTestRunner()
 			.ModifyWorld([&](UWorld* World)
 			{
-				TestActor = World->SpawnActor<ASnapshotTestActor>();
+				TestActor = ASnapshotTestActor::Spawn(World, "TestActor");
 				ActorSerializer->TestActor = TestActor;
 				SubobjectSerializer->TestActor = TestActor;
 			
@@ -471,7 +470,7 @@ namespace UE::LevelSnapshots::Private::Tests
 		FSnapshotTestRunner()
 			.ModifyWorld([&](UWorld* World)
 			{
-				TestActor = World->SpawnActor<ASnapshotTestActor>();
+				TestActor = ASnapshotTestActor::Spawn(World, "TestActor");
 				ActorSerializer->TestActor = TestActor;
 				SubobjectSerializer->TestActor = TestActor;
 			
@@ -514,8 +513,8 @@ namespace UE::LevelSnapshots::Private::Tests
 		FSnapshotTestRunner()
 			.ModifyWorld([&](UWorld* World)
 			{
-				ChangedPropertiesActor = World->SpawnActor<ASnapshotTestActor>();
-				NulledActor = World->SpawnActor<ASnapshotTestActor>();
+				ChangedPropertiesActor = ASnapshotTestActor::Spawn(World, "ChangedPropertiesActor");
+				NulledActor = ASnapshotTestActor::Spawn(World, "NulledActor");
 
 				ChangedPropertiesActor->InstancedOnlySubobject_DefaultSubobject->FloatProperty = 42.f;
 				ChangedPropertiesActor->NakedSubobject_DefaultSubobject->FloatProperty = 21.f;
@@ -601,7 +600,7 @@ namespace UE::LevelSnapshots::Private::Tests
 		FSnapshotTestRunner()
 			.ModifyWorld([&](UWorld* World)
 			{
-				Stub->TestActor = World->SpawnActor<ASnapshotTestActor>();
+				Stub->TestActor = ASnapshotTestActor::Spawn(World, "TestActor");
 			
 				Stub->TestActor->InstancedOnlySubobject_DefaultSubobject->FloatProperty = 42.f;
 			})
@@ -638,7 +637,7 @@ namespace UE::LevelSnapshots::Private::Tests
 		FSnapshotTestRunner()
 			.ModifyWorld([&](UWorld* World)
 			{
-				Stub->TestActor = World->SpawnActor<ASnapshotTestActor>();
+				Stub->TestActor = ASnapshotTestActor::Spawn(World, "TestActor");
 			})
 			.TakeSnapshot()
 			.FilterProperties(Stub->TestActor, [&](const FPropertySelectionMap& PropertySelectionMap)
@@ -663,7 +662,7 @@ namespace UE::LevelSnapshots::Private::Tests
 		FSnapshotTestRunner()
 			.ModifyWorld([&](UWorld* World)
 			{
-				Stub->TestActor = World->SpawnActor<ASnapshotTestActor>();
+				Stub->TestActor = ASnapshotTestActor::Spawn(World, "TestActor");
 			
 				Stub->TestActor->InstancedOnlySubobject_DefaultSubobject->FloatProperty = 42.f;
 				Stub->TestActor->IntProperty = 42;
@@ -784,7 +783,7 @@ namespace UE::LevelSnapshots::Private::Tests
 		FSnapshotTestRunner()
 			.ModifyWorld([&](UWorld* World)
 			{
-				Stub->TestActor = World->SpawnActor<ASnapshotTestActor>();
+				Stub->TestActor = ASnapshotTestActor::Spawn(World, "TestActor");
 				Stub->TestActor->IntProperty = 21;
 				Stub->TestActor->AllocateNonReflectedSubobject();
 				Stub->TestActor->NonReflectedSubobject->FloatProperty = 42.f;
diff --git a/Engine/Plugins/VirtualProduction/LevelSnapshots/Source/LevelSnapshots/Private/Tests/RegressionTests.cpp b/Engine/Plugins/VirtualProduction/LevelSnapshots/Source/LevelSnapshots/Private/Tests/RegressionTests.cpp
--- a/Engine/Plugins/VirtualProduction/LevelSnapshots/Source/LevelSnapshots/Private/Tests/RegressionTests.cpp
+++ b/Engine/Plugins/VirtualProduction/LevelSnapshots/Source/LevelSnapshots/Private/Tests/RegressionTests.cpp
@@ -84,11 +84,11 @@ namespace UE::LevelSnapshots::Private::Tests
 		FSnapshotTestRunner()
 			.ModifyWorld([&](UWorld* World)
 			{
-				CustomBlockAllToOverlapAll = World->SpawnActor<ASnapshotTestActor>();
-				DefaultWithChangedIgnoredProperties = World->SpawnActor<ASnapshotTestActor>();
-				DefaultToCustom = World->SpawnActor<ASnapshotTestActor>();
-				CustomToDefault = World->SpawnActor<ASnapshotTestActor>();
-				ChangeCollisionProfile = World->SpawnActor<ASnapshotTestActor>();
+				CustomBlockAllToOverlapAll = ASnapshotTestActor::Spawn(World, "CustomBlockAllToOverlapAll");
+				DefaultWithChangedIgnoredProperties = ASnapshotTestActor::Spawn(World, "DefaultWithChangedIgnoredProperties");
+				DefaultToCustom = ASnapshotTestActor::Spawn(World, "DefaultToCustom");
+				CustomToDefault = ASnapshotTestActor::Spawn(World, "CustomToDefault");
+				ChangeCollisionProfile = ASnapshotTestActor::Spawn(World, "ChangeCollisionProfile");
 
 				CustomBlockAllToOverlapAll->StaticMeshComponent->bUseDefaultCollision = false;
 				CustomBlockAllToOverlapAll->StaticMeshComponent->SetCollisionResponseToAllChannels(ECR_Block);
diff --git a/Engine/Plugins/VirtualProduction/LevelSnapshots/Source/LevelSnapshots/Private/Tests/SubobjectTests.cpp b/Engine/Plugins/VirtualProduction/LevelSnapshots/Source/LevelSnapshots/Private/Tests/SubobjectTests.cpp
--- a/Engine/Plugins/VirtualProduction/LevelSnapshots/Source/LevelSnapshots/Private/Tests/SubobjectTests.cpp
+++ b/Engine/Plugins/VirtualProduction/LevelSnapshots/Source/LevelSnapshots/Private/Tests/SubobjectTests.cpp
@@ -12,15 +12,6 @@
 #include "Misc/AutomationTest.h"
 #include "Misc/ScopeExit.h"
 
-// TODO: FExistingSubobjectHasDifferentClass > Change a subobject's class and try to restore into it
-// TODO: FCollectionDoesNotContainSameSubobjects > Generates diff
-// TODO: FCollectionOfNonSubobjectsDiffesCorrectly > Generates diff when collection contains non-subobjects
-// TODO: Mixed references in collection: own subobjects and external subobjects
-// TODO: What happens when TArray<UObject*> contains an object with disallowed class > Do not serialize into it. Only try resolving to it.
-// TODO: Unrelated classes but same object name in collection and one of the classes is not allowed. What happens?
-
-// TODO: FSubobjectWithSamePropertiesAndDifferentNames does not diff
-
 namespace UE::LevelSnapshots::Private::Tests
 {
 	IMPLEMENT_SIMPLE_AUTOMATION_TEST(FRestoreSubobjectProperties, "VirtualProduction.LevelSnapshots.Snapshot.Subobject.RestoreSubobjectProperties", (EAutomationTestFlags::ApplicationContextMask | EAutomationTestFlags::EngineFilter));
@@ -34,7 +25,7 @@ namespace UE::LevelSnapshots::Private::Tests
 				FSnapshotTestRunner()
 					.ModifyWorld([&](UWorld* World)
 					{
-						Actor = World->SpawnActor<ASnapshotTestActor>();
+						Actor = ASnapshotTestActor::Spawn(World, "TestActor");
 						Actor->AllocateSubobjects();
 
 						Actor->EditableInstancedSubobject_DefaultSubobject->IntProperty					= 1;
@@ -142,7 +133,7 @@ namespace UE::LevelSnapshots::Private::Tests
 				FSnapshotTestRunner()
 					.ModifyWorld([&](UWorld* World)
 					{
-						Actor = World->SpawnActor<ASnapshotTestActor>();
+						Actor = ASnapshotTestActor::Spawn(World, "TestActor");
 						Actor->AllocateSubobjects();
 
 						Actor->EditableInstancedSubobject_DefaultSubobject->IntProperty					= 1;
@@ -265,7 +256,7 @@ namespace UE::LevelSnapshots::Private::Tests
 				FSnapshotTestRunner()
 					.ModifyWorld([&](UWorld* World)
 					{
-						Actor = World->SpawnActor<ASnapshotTestActor>();
+						Actor = ASnapshotTestActor::Spawn(World, "TestActor");
 						Actor->AllocateSubobjects();
 
 						Actor->EditableInstancedSubobject_DefaultSubobject->IntProperty				= 1;
@@ -338,7 +329,7 @@ namespace UE::LevelSnapshots::Private::Tests
 				FSnapshotTestRunner()
 					.ModifyWorld([&](UWorld* World)
 					{
-						Actor = World->SpawnActor<ASnapshotTestActor>();
+						Actor = ASnapshotTestActor::Spawn(World, "TestActor");
 						Actor->AllocateSubobjects();
 
 						Actor->EditableInstancedSubobjectArray_OptionalSubobject[0]->IntProperty		= 1;
@@ -384,7 +375,7 @@ namespace UE::LevelSnapshots::Private::Tests
 				FSnapshotTestRunner()
 					.ModifyWorld([&](UWorld* World)
 					{
-						Actor = World->SpawnActor<ASnapshotTestActor>();
+						Actor = ASnapshotTestActor::Spawn(World, "TestActor");
 						Actor->AllocateSubobjects();
 
 						Actor->ObjectReference											= Actor->EditableInstancedSubobjectArray_OptionalSubobject[0];
@@ -493,7 +484,7 @@ namespace UE::LevelSnapshots::Private::Tests
 		FSnapshotTestRunner()
 			.ModifyWorld([&](UWorld* World)
 			{
-				Actor = World->SpawnActor<ASnapshotTestActor>();
+				Actor = ASnapshotTestActor::Spawn(World, "TestActor");
 				Actor->AllocateSubobjects();
 
 				Subobject = Actor->EditableInstancedSubobject_DefaultSubobject;
@@ -528,7 +519,7 @@ namespace UE::LevelSnapshots::Private::Tests
 		FSnapshotTestRunner()
 			.ModifyWorld([&](UWorld* World)
 			{
-				Actor = World->SpawnActor<ASnapshotTestActor>();
+				Actor = ASnapshotTestActor::Spawn(World, "TestActor");
 				Actor->AllocateSubobjects();
 
 				Actor->EditableInstancedSubobject_DefaultSubobject->IntProperty = 21;
@@ -576,9 +567,9 @@ namespace UE::LevelSnapshots::Private::Tests
 				FSnapshotTestRunner()
 					.ModifyWorld([&](UWorld* World)
 					{
-						ActorA = World->SpawnActor<ASnapshotTestActor>();
+						ActorA = ASnapshotTestActor::Spawn(World, "ActorA");
 						ActorA->AllocateSubobjects();
-						ActorB = World->SpawnActor<ASnapshotTestActor>();
+						ActorB = ASnapshotTestActor::Spawn(World, "ActorB");
 						ActorB->AllocateSubobjects();
 
 						ActorA->EditableInstancedSubobject_DefaultSubobject->IntProperty = 21;
@@ -651,8 +642,8 @@ namespace UE::LevelSnapshots::Private::Tests
 		FSnapshotTestRunner()
 			.ModifyWorld([&](UWorld* World)
 			{
-				ActorA = World->SpawnActor<ASnapshotTestActor>();
-				ActorB = World->SpawnActor<ASnapshotTestActor>();
+				ActorA = ASnapshotTestActor::Spawn(World, "ActorA");
+				ActorB = ASnapshotTestActor::Spawn(World, "ActorB");
 
 				ActorA->ObjectReference = ActorB->TestComponent;
 			})
@@ -672,8 +663,8 @@ namespace UE::LevelSnapshots::Private::Tests
 		FSnapshotTestRunner()
 			.ModifyWorld([&](UWorld* World)
 			{
-				ActorA = World->SpawnActor<ASnapshotTestActor>();
-				ActorB = World->SpawnActor<ASnapshotTestActor>();
+				ActorA = ASnapshotTestActor::Spawn(World, "ActorA");
+				ActorB = ASnapshotTestActor::Spawn(World, "ActorB");
 
 				ActorA->ObjectReference = ActorB->TestComponent;
 			})
@@ -735,9 +726,9 @@ namespace UE::LevelSnapshots::Private::Tests
 		FSnapshotTestRunner()
 			.ModifyWorld([&](UWorld* World)
 			{
-				ActorA = World->SpawnActor<ASnapshotTestActor>();
+				ActorA = ASnapshotTestActor::Spawn(World, "ActorA");
 				ActorA->AllocateSubobjects();
-				ActorB = World->SpawnActor<ASnapshotTestActor>();
+				ActorB = ASnapshotTestActor::Spawn(World, "ActorB");
 				ActorB->AllocateSubobjects();
 
 				ActorA->EditableInstancedSubobject_DefaultSubobject->IntProperty			= 21;
@@ -957,10 +948,10 @@ namespace UE::LevelSnapshots::Private::Tests
 		FSnapshotTestRunner()
 			.ModifyWorld([&](UWorld* World)
 			{
-				ReferenceExternalSubobjects = World->SpawnActor<ASnapshotTestActor>();
+				ReferenceExternalSubobjects = ASnapshotTestActor::Spawn(World, "ReferenceExternalSubobjects");
 				ReferenceExternalSubobjects->AllocateSubobjects();
 			
-				ReferencedObject = World->SpawnActor<ASnapshotTestActor>();
+				ReferencedObject = ASnapshotTestActor::Spawn(World, "ReferencedObject");
 				ReferencedObject->AllocateSubobjects();
 				DestroyedReferencedSubobject = NewObject<USubobject>(ReferencedObject, USubobject::StaticClass(), TEXT("DestroyedReferencedSubobject"));
 				DestroyedReferencedSubobject->FloatProperty = 42.f;
diff --git a/Engine/Plugins/VirtualProduction/LevelSnapshots/Source/LevelSnapshots/Private/Tests/Types/SnapshotTestActor.cpp b/Engine/Plugins/VirtualProduction/LevelSnapshots/Source/LevelSnapshots/Private/Tests/Types/SnapshotTestActor.cpp
--- a/Engine/Plugins/VirtualProduction/LevelSnapshots/Source/LevelSnapshots/Private/Tests/Types/SnapshotTestActor.cpp
+++ b/Engine/Plugins/VirtualProduction/LevelSnapshots/Source/LevelSnapshots/Private/Tests/Types/SnapshotTestActor.cpp
@@ -21,6 +21,7 @@ ASnapshotTestActor* ASnapshotTestActor::Spawn(UWorld* World, FName Name)
 	Params.Name = Name;
 	Params.bNoFail = true;
 	Params.SpawnCollisionHandlingOverride = ESpawnActorCollisionHandlingMethod::AlwaysSpawn;
+	Params.NameMode = FActorSpawnParameters::ESpawnActorNameMode::Requested;
 	return World->SpawnActor<ASnapshotTestActor>(Params);
 }
 
diff --git a/Engine/Plugins/VirtualProduction/LevelSnapshots/Source/LevelSnapshots/Private/Tests/Types/SnapshotTestActor.h b/Engine/Plugins/VirtualProduction/LevelSnapshots/Source/LevelSnapshots/Private/Tests/Types/SnapshotTestActor.h
--- a/Engine/Plugins/VirtualProduction/LevelSnapshots/Source/LevelSnapshots/Private/Tests/Types/SnapshotTestActor.h
+++ b/Engine/Plugins/VirtualProduction/LevelSnapshots/Source/LevelSnapshots/Private/Tests/Types/SnapshotTestActor.h
@@ -74,7 +74,7 @@ class ASnapshotTestActor : public ACharacter
 	GENERATED_BODY()
 public:
 
-	static ASnapshotTestActor* Spawn(UWorld* World, FName Name);
+	static ASnapshotTestActor* Spawn(UWorld* World, FName Name = FName());
 	
 	ASnapshotTestActor();
 
