Level Snapshots: Fix transform sometimes not getting restored correctly when attaching actor

diff --git a/Engine/Plugins/VirtualProduction/LevelSnapshots/Source/LevelSnapshots/Private/Data/Util/Restoration/RestoreActorUtil.cpp b/Engine/Plugins/VirtualProduction/LevelSnapshots/Source/LevelSnapshots/Private/Data/Util/Restoration/RestoreActorUtil.cpp
--- a/Engine/Plugins/VirtualProduction/LevelSnapshots/Source/LevelSnapshots/Private/Data/Util/Restoration/RestoreActorUtil.cpp
+++ b/Engine/Plugins/VirtualProduction/LevelSnapshots/Source/LevelSnapshots/Private/Data/Util/Restoration/RestoreActorUtil.cpp
@@ -97,9 +97,10 @@ namespace UE::LevelSnapshots::Private::Internal::Restore
 			AttachParentAfterRestore->GetRootComponent()->DetachFromComponent(FDetachmentTransformRules::KeepRelativeTransform);
 		}
 
-		OriginalActor->RerunConstructionScripts();
 		// Update component state, e.g. render state if intensity for lights was changed. Also avoids us having to call PostEditChangeProperty on every property.
+		// ReregisterAllComponents must be called before RerunConstructionScripts because it processes calls we made to SetupAttachment
 		OriginalActor->ReregisterAllComponents();
+		OriginalActor->RerunConstructionScripts();
 
 #if WITH_EDITOR
 		// Update World Outliner. Usually called by USceneComponent::AttachToComponent.
@@ -124,6 +125,21 @@ namespace UE::LevelSnapshots::Private::Internal::Restore
 			SceneComponent->DetachFromComponent(FDetachmentTransformRules::KeepRelativeTransform);
 		}
 	}
+
+	static void UpdateAttachParentAttachChildren(UActorComponent* Original)
+	{
+		// Original->AttachParent had its value serialized but new attach parent's AttachChildren must be updated, too
+		USceneComponent* Component = Cast<USceneComponent>(Original);
+		const bool bNeedsToUpdateParentAttachChildren = Component && Component->GetAttachParent()
+			&& !Component->GetAttachParent()->GetAttachChildren().Contains(Component);
+		if (bNeedsToUpdateParentAttachChildren)
+		{
+			// Hacky way of updating AttachChildren since there is no direct way...
+			Component->UnregisterComponent();
+			Component->SetupAttachment(Component->GetAttachParent());
+			// ReregisterAllComponents will be called after the actor is done restoring which will fix up everything
+		}
+	}
 }
 
 void UE::LevelSnapshots::Private::RestoreIntoExistingWorldActor(AActor* OriginalActor, FActorSnapshotData& ActorData, FWorldSnapshotData& WorldData, FSnapshotDataCache& Cache, UPackage* InLocalisationSnapshotPackage, const FPropertySelectionMap& SelectedProperties)
@@ -150,6 +166,7 @@ void UE::LevelSnapshots::Private::RestoreIntoExistingWorldActor(AActor* Original
 		{
 			Internal::Restore::PreventAttachParentInfiniteRecursion(Original, *ComponentSelectedProperties);
 			FApplySnapshotToEditorArchive::ApplyToExistingEditorWorldObject(SerializedCompData, WorldData, Cache, Original, Deserialized, SelectedProperties);
+			Internal::Restore::UpdateAttachParentAttachChildren(Original);
 		};
 	};
 	Internal::Restore::DeserializeIntoEditorWorldActor(OriginalActor, ActorData, WorldData, Cache, InLocalisationSnapshotPackage, DeserializeActor, DeserializeComponent);
@@ -172,6 +189,7 @@ void UE::LevelSnapshots::Private::RestoreIntoRecreatedEditorWorldActor(AActor* O
 	{
 		const FRestoreObjectScope FinishRestore = PreSubobjectRestore_EditorWorld(Deserialized, Original, WorldData, Cache, SelectedProperties, InLocalisationSnapshotPackage);
 		FApplySnapshotToEditorArchive::ApplyToRecreatedEditorWorldObject(SerializedCompData, WorldData, Cache, Original, SelectedProperties); 
+		Internal::Restore::UpdateAttachParentAttachChildren(Original);
 	};
 	Internal::Restore::DeserializeIntoEditorWorldActor(OriginalActor, ActorData, WorldData, Cache, InLocalisationSnapshotPackage, DeserializeActor, DeserializeComponent);
 
