Rename "Load filter" to "Load preset"

diff --git a/Engine/Plugins/VirtualProduction/LevelSnapshots/Source/LevelSnapshotsEditor/Private/Widgets/Filter/SSaveAndLoadFilters.cpp b/Engine/Plugins/VirtualProduction/LevelSnapshots/Source/LevelSnapshotsEditor/Private/Widgets/Filter/SSaveAndLoadFilters.cpp
--- a/Engine/Plugins/VirtualProduction/LevelSnapshots/Source/LevelSnapshotsEditor/Private/Widgets/Filter/SSaveAndLoadFilters.cpp
+++ b/Engine/Plugins/VirtualProduction/LevelSnapshots/Source/LevelSnapshotsEditor/Private/Widgets/Filter/SSaveAndLoadFilters.cpp
@@ -144,7 +144,7 @@ TSharedRef<SWidget> SSaveAndLoadFilters::GenerateSaveLoadMenu()
 		)	
     );
 
-	MenuBuilder.BeginSection(NAME_None, LOCTEXT("LoadFilter_MenuSection", "Load filter"));
+	MenuBuilder.BeginSection(NAME_None, LOCTEXT("LoadPreset_MenuSection", "Load preset"));
 	{
 		FContentBrowserModule& ContentBrowserModule = FModuleManager::Get().LoadModuleChecked<FContentBrowserModule>(TEXT("ContentBrowser"));
 
