Level Snapshots:

-Fixed Duplicated LOCTEXT Keys (Release-5.0 - Duplicated keys for localization - UE - Virtual Production)

diff --git a/Engine/Plugins/VirtualProduction/LevelSnapshots/Source/LevelSnapshots/Private/Data/Util/Restoration/WorldDataUtil.cpp b/Engine/Plugins/VirtualProduction/LevelSnapshots/Source/LevelSnapshots/Private/Data/Util/Restoration/WorldDataUtil.cpp
--- a/Engine/Plugins/VirtualProduction/LevelSnapshots/Source/LevelSnapshots/Private/Data/Util/Restoration/WorldDataUtil.cpp
+++ b/Engine/Plugins/VirtualProduction/LevelSnapshots/Source/LevelSnapshots/Private/Data/Util/Restoration/WorldDataUtil.cpp
@@ -124,7 +124,7 @@ namespace UE::LevelSnapshots::Private::Internal
 
 	static void CaptureSnapshotData(const TArray<AActor*>& ActorsInWorld, FWorldSnapshotData& SnapshotData)
 	{
-		FScopedSlowTask CaptureData(ActorsInWorld.Num(), LOCTEXT("TakeSnapshotKey", "Capturing data"));
+		FScopedSlowTask CaptureData(ActorsInWorld.Num(), LOCTEXT("CapturingWorldData", "Capturing data"));
 		CaptureData.MakeDialogDelayed(1.f);
 		
 		const bool bShouldLog = ConsoleVariables::CVarLogTimeTakingSnapshots.GetValueOnAnyThread();
@@ -143,7 +143,7 @@ namespace UE::LevelSnapshots::Private::Internal
 
 	static void ComputeActorHashes(const TArray<AActor*>& ActorsInWorld, FWorldSnapshotData& SnapshotData)
 	{
-		FScopedSlowTask ComputeHash(1.f, LOCTEXT("TakeSnapshotKey", "Computing hashes"));
+		FScopedSlowTask ComputeHash(1.f, LOCTEXT("ComputingDataHashes", "Computing hashes"));
 		ComputeHash.MakeDialogDelayed(1.f);
 		
 		// Hashing takes about half of the time
@@ -553,4 +553,4 @@ void UE::LevelSnapshots::Private::SerializeClassDefaultsInto(FWorldSnapshotData&
 			);
 }
 
-#undef LOCTEXT_NAMESPACE
\ No newline at end of file
+#undef LOCTEXT_NAMESPACE
diff --git a/Engine/Plugins/VirtualProduction/LevelSnapshots/Source/LevelSnapshotsEditor/Private/Util/TakeSnapshotUtil.cpp b/Engine/Plugins/VirtualProduction/LevelSnapshots/Source/LevelSnapshotsEditor/Private/Util/TakeSnapshotUtil.cpp
--- a/Engine/Plugins/VirtualProduction/LevelSnapshots/Source/LevelSnapshotsEditor/Private/Util/TakeSnapshotUtil.cpp
+++ b/Engine/Plugins/VirtualProduction/LevelSnapshots/Source/LevelSnapshotsEditor/Private/Util/TakeSnapshotUtil.cpp
@@ -188,7 +188,7 @@ ULevelSnapshot* SnapshotEditor::TakeLevelSnapshotAndSaveToDisk(UWorld* World, co
     	{
     		NotificationItem->SetText(
     			FText::Format(
-    				NSLOCTEXT("LevelSnapshots", "NotificationFormatText_CreateSnapshotSuccess", "Failed to create Level Snapshot \"{0}\". Check the file name."), FText::FromString(FileName)));
+    				NSLOCTEXT("LevelSnapshots", "NotificationFormatText_CreateSnapshotFailure", "Failed to create Level Snapshot \"{0}\". Check the file name."), FText::FromString(FileName)));
     		NotificationItem->SetCompletionState(SNotificationItem::CS_Fail);
     	}
 
@@ -232,4 +232,4 @@ void SnapshotEditor::GenerateThumbnailForSnapshotAsset(ULevelSnapshot* Snapshot)
 	}
 }
 
-#undef LOCTEXT_NAMESPACE
\ No newline at end of file
+#undef LOCTEXT_NAMESPACE
