LevelSnapshots: Fix cache not getting invalidated
Prevent crash for saving nDisplay cluster nodes

diff --git a/Engine/Plugins/VirtualProduction/LevelSnapshots/Source/LevelSnapshots/Private/Data/CustomSerialization/CustomSerializationDataManager.cpp b/Engine/Plugins/VirtualProduction/LevelSnapshots/Source/LevelSnapshots/Private/Data/CustomSerialization/CustomSerializationDataManager.cpp
--- a/Engine/Plugins/VirtualProduction/LevelSnapshots/Source/LevelSnapshots/Private/Data/CustomSerialization/CustomSerializationDataManager.cpp
+++ b/Engine/Plugins/VirtualProduction/LevelSnapshots/Source/LevelSnapshots/Private/Data/CustomSerialization/CustomSerializationDataManager.cpp
@@ -116,12 +116,12 @@ int32 UE::LevelSnapshots::Private::FCustomSerializationDataReader::AddSubobjectS
 
 TSharedPtr<UE::LevelSnapshots::ISnapshotSubobjectMetaData> UE::LevelSnapshots::Private::FCustomSerializationDataReader::GetSubobjectMetaData(int32 Index)
 {
-	return CachedSubobjectMetaData[Index];
+	return ensure(CachedSubobjectMetaData.IsValidIndex(Index)) ? CachedSubobjectMetaData[Index] : nullptr;
 }
 
 const TSharedPtr<UE::LevelSnapshots::ISnapshotSubobjectMetaData> UE::LevelSnapshots::Private::FCustomSerializationDataReader::GetSubobjectMetaData(int32 Index) const
 {
-	return CachedSubobjectMetaData[Index];
+	return ensure(CachedSubobjectMetaData.IsValidIndex(Index)) ? CachedSubobjectMetaData[Index] : nullptr;
 }
 
 int32 UE::LevelSnapshots::Private::FCustomSerializationDataReader::GetNumSubobjects() const
@@ -167,8 +167,11 @@ void UE::LevelSnapshots::Private::FCustomSerializationDataWriter::WriteObjectAnn
 
 int32 UE::LevelSnapshots::Private::FCustomSerializationDataWriter::AddSubobjectSnapshot(UObject* Subobject)
 {
-	if (!ensure(Subobject && Subobject->IsIn(SerializedObject)))
+	if (!ensure(Subobject) || !ensure(Subobject->IsIn(SerializedObject)))
 	{
+		const FString SubobjectPath = Subobject ? Subobject->GetPathName() : FString("None");
+		const FString SerializedObjectPath = SerializedObject ? SerializedObject->GetPathName() : FString("None");
+		UE_LOG(LogLevelSnapshots, Error, TEXT("%s is not in %s!"), *SubobjectPath, *SerializedObjectPath);
 		return INDEX_NONE;
 	}
 
@@ -180,16 +183,14 @@ int32 UE::LevelSnapshots::Private::FCustomSerializationDataWriter::AddSubobjectS
 		return Data.ObjectPathIndex == ObjectIndex;
 	});
 
-	if (ObjectIndex != INDEX_NONE)
+	if (!ensure(ObjectIndex == INDEX_NONE))
 	{
-		UE_LOG(LogLevelSnapshots, Error, TEXT("You tried to register an object which was already found by standard Level Snapshot serialisation. Is your subobject referenced by an property with a CPF_Edit flag?"));
-		UE_DEBUG_BREAK();
+		UE_LOG(LogLevelSnapshots, Error, TEXT("You tried to register an object (%s) which was already found by standard Level Snapshot serialisation. Is your subobject referenced by an property with a CPF_Edit flag?"), *Subobject->GetPathName());
 		return ExistingSubobjectIndex;
 	}
-	if (ExistingSubobjectIndex != INDEX_NONE)
+	if (!ensure(ExistingSubobjectIndex == INDEX_NONE))
 	{
-		UE_LOG(LogLevelSnapshots, Error, TEXT("You tried to register the same subobject twice."));
-		UE_DEBUG_BREAK();
+		UE_LOG(LogLevelSnapshots, Error, TEXT("You tried to register the same subobject (%s) twice."), *Subobject->GetPathName());
 		return ExistingSubobjectIndex;
 	}
 
diff --git a/Engine/Plugins/VirtualProduction/LevelSnapshots/Source/LevelSnapshots/Private/Data/LevelSnapshot.cpp b/Engine/Plugins/VirtualProduction/LevelSnapshots/Source/LevelSnapshots/Private/Data/LevelSnapshot.cpp
--- a/Engine/Plugins/VirtualProduction/LevelSnapshots/Source/LevelSnapshots/Private/Data/LevelSnapshot.cpp
+++ b/Engine/Plugins/VirtualProduction/LevelSnapshots/Source/LevelSnapshots/Private/Data/LevelSnapshot.cpp
@@ -391,7 +391,7 @@ void ULevelSnapshot::ClearCache()
 void ULevelSnapshot::ClearCachedDiffFlag(UObject* ModifiedObject)
 {
 	AActor* AsActor = ModifiedObject->IsA<AActor>() ? Cast<AActor>(ModifiedObject) : ModifiedObject->GetTypedOuter<AActor>(); 
-	if (AsActor && SerializedData.HasMatchingSavedActor(ModifiedObject))
+	if (AsActor && SerializedData.HasMatchingSavedActor(AsActor))
 	{
 		CachedDiffedActors.FindOrAdd(AsActor) = ECachedDiffResult::NotInitialised;
 	}
diff --git a/Engine/Plugins/VirtualProduction/LevelSnapshots/Source/nDisplaySupport/Private/Helpers/MapSubobjectSerializer.h b/Engine/Plugins/VirtualProduction/LevelSnapshots/Source/nDisplaySupport/Private/Helpers/MapSubobjectSerializer.h
--- a/Engine/Plugins/VirtualProduction/LevelSnapshots/Source/nDisplaySupport/Private/Helpers/MapSubobjectSerializer.h
+++ b/Engine/Plugins/VirtualProduction/LevelSnapshots/Source/nDisplaySupport/Private/Helpers/MapSubobjectSerializer.h
@@ -3,6 +3,7 @@
 #pragma once
 
 #include "CoreMinimal.h"
+#include "LevelSnapshotsLog.h"
 #include "Interfaces/ICustomObjectSnapshotSerializer.h"
 #include "Params/ObjectSnapshotSerializationData.h"
 #include "UObject/UnrealType.h"
@@ -87,13 +88,20 @@ namespace UE::LevelSnapshots::nDisplay::Private
 				if (ensure(IsValid(Value)))
 				{
 					const int32 ViewportIndex = DataStorage.AddSubobjectSnapshot(Value);
-					DataStorage.GetSubobjectMetaData(ViewportIndex)
-						->WriteObjectAnnotation(FObjectAnnotator::CreateLambda([&Key](FArchive& Writer)
+					TSharedPtr<ISnapshotSubobjectMetaData> MetaData = DataStorage.GetSubobjectMetaData(ViewportIndex);
+					if (ensureMsgf(MetaData, TEXT("AddSubobjectSnapshot failed unexpectingly")))
+					{
+						MetaData->WriteObjectAnnotation(FObjectAnnotator::CreateLambda([&Key](FArchive& Writer)
 						{
 							FString ViewportKey = Key;
 							Writer << ViewportKey;
 						}));
+						
+						continue;
+					}
 				}
+
+				UE_LOG(LogLevelSnapshots, Error, TEXT("Failed to save %s that is supposed to be a subobject of %s"), *Value->GetPathName(), *EditorObject->GetPathName());
 			}
 		}
 		virtual UObject* FindOrRecreateSubobjectInSnapshotWorld(UObject* SnapshotObject, const ISnapshotSubobjectMetaData& ObjectData, const ICustomSnapshotSerializationData& DataStorage) override
