Level Snapshots:

-Fix for: Create Level Snapshot window has to be closed and reopened for the Default Level Snapshot Name to be considered

diff --git a/Engine/Plugins/VirtualProduction/LevelSnapshots/Source/LevelSnapshotsEditor/Private/LevelSnapshotsEditorModule.cpp b/Engine/Plugins/VirtualProduction/LevelSnapshots/Source/LevelSnapshotsEditor/Private/LevelSnapshotsEditorModule.cpp
--- a/Engine/Plugins/VirtualProduction/LevelSnapshots/Source/LevelSnapshotsEditor/Private/LevelSnapshotsEditorModule.cpp
+++ b/Engine/Plugins/VirtualProduction/LevelSnapshots/Source/LevelSnapshotsEditor/Private/LevelSnapshotsEditorModule.cpp
@@ -150,7 +150,7 @@ void FLevelSnapshotsEditorModule::OpenSnapshotsEditor()
 
 bool FLevelSnapshotsEditorModule::HandleModifiedProjectSettings()
 {
-	ULevelSnapshotsEditorSettings* Settings = ULevelSnapshotsEditorSettings::Get();
+	ULevelSnapshotsEditorSettings* Settings = ULevelSnapshotsEditorSettings::GetMutable();
 	Settings->ValidateRootLevelSnapshotSaveDirAsGameContentRelative();
 	Settings->SanitizeAllProjectSettingsPaths(true);
 		
diff --git a/Engine/Plugins/VirtualProduction/LevelSnapshots/Source/LevelSnapshotsEditor/Private/LevelSnapshotsEditorModule.h b/Engine/Plugins/VirtualProduction/LevelSnapshots/Source/LevelSnapshotsEditor/Private/LevelSnapshotsEditorModule.h
--- a/Engine/Plugins/VirtualProduction/LevelSnapshots/Source/LevelSnapshotsEditor/Private/LevelSnapshotsEditorModule.h
+++ b/Engine/Plugins/VirtualProduction/LevelSnapshots/Source/LevelSnapshotsEditor/Private/LevelSnapshotsEditorModule.h
@@ -29,7 +29,7 @@ public:
 	//~ End IModuleInterface Interface
 
 	static bool GetUseCreationForm() { return ULevelSnapshotsEditorSettings::Get()->bUseCreationForm; }
-	static void SetUseCreationForm(bool bInUseCreationForm) { ULevelSnapshotsEditorSettings::Get()->bUseCreationForm = bInUseCreationForm; }
+	static void SetUseCreationForm(bool bInUseCreationForm) { ULevelSnapshotsEditorSettings::GetMutable()->bUseCreationForm = bInUseCreationForm; }
 	static void ToggleUseCreationForm() { SetUseCreationForm(!GetUseCreationForm()); }
 
 	void OpenLevelSnapshotsDialogWithAssetSelected(const FAssetData& InAssetData);
diff --git a/Engine/Plugins/VirtualProduction/LevelSnapshots/Source/LevelSnapshotsEditor/Private/LevelSnapshotsEditorSettings.cpp b/Engine/Plugins/VirtualProduction/LevelSnapshots/Source/LevelSnapshotsEditor/Private/LevelSnapshotsEditorSettings.cpp
--- a/Engine/Plugins/VirtualProduction/LevelSnapshots/Source/LevelSnapshotsEditor/Private/LevelSnapshotsEditorSettings.cpp
+++ b/Engine/Plugins/VirtualProduction/LevelSnapshots/Source/LevelSnapshotsEditor/Private/LevelSnapshotsEditorSettings.cpp
@@ -6,7 +6,12 @@
 #include "Application/SlateApplicationBase.h"
 #include "HAL/PlatformApplicationMisc.h"
 
-ULevelSnapshotsEditorSettings* ULevelSnapshotsEditorSettings::Get()
+const ULevelSnapshotsEditorSettings* ULevelSnapshotsEditorSettings::Get()
+{
+	return GetDefault<ULevelSnapshotsEditorSettings>();
+}
+
+ULevelSnapshotsEditorSettings* ULevelSnapshotsEditorSettings::GetMutable()
 {
 	return GetMutableDefault<ULevelSnapshotsEditorSettings>();
 }
@@ -16,7 +21,6 @@ ULevelSnapshotsEditorSettings::ULevelSnapshotsEditorSettings(const FObjectInitia
 	RootLevelSnapshotSaveDir.Path = "/Game/LevelSnapshots";
 	LevelSnapshotSaveDir = "{map}/{year}-{month}-{day}";
 	DefaultLevelSnapshotName = "{map}_{user}_{time}";
-	LevelSnapshotNameOverride = DefaultLevelSnapshotName;
 
 	bEnableLevelSnapshotsToolbarButton = true;
 	bUseCreationForm = true;
@@ -48,7 +52,7 @@ void ULevelSnapshotsEditorSettings::SetLastCreationWindowSize(const FVector2D In
 
 const FString& ULevelSnapshotsEditorSettings::GetNameOverride() const
 {
-	return LevelSnapshotNameOverride;
+	return LevelSnapshotNameOverride.IsSet() ? LevelSnapshotNameOverride.Get() : DefaultLevelSnapshotName;
 }
 
 void ULevelSnapshotsEditorSettings::SetNameOverride(const FString& InName)
@@ -121,5 +125,5 @@ FText ULevelSnapshotsEditorSettings::ParseLevelSnapshotsTokensInText(const FText
 
 bool ULevelSnapshotsEditorSettings::IsNameOverridden() const
 {
-	return !LevelSnapshotNameOverride.Equals(DefaultLevelSnapshotName);
+	return LevelSnapshotNameOverride.IsSet();
 }
diff --git a/Engine/Plugins/VirtualProduction/LevelSnapshots/Source/LevelSnapshotsEditor/Private/LevelSnapshotsEditorSettings.h b/Engine/Plugins/VirtualProduction/LevelSnapshots/Source/LevelSnapshotsEditor/Private/LevelSnapshotsEditorSettings.h
--- a/Engine/Plugins/VirtualProduction/LevelSnapshots/Source/LevelSnapshotsEditor/Private/LevelSnapshotsEditorSettings.h
+++ b/Engine/Plugins/VirtualProduction/LevelSnapshots/Source/LevelSnapshotsEditor/Private/LevelSnapshotsEditorSettings.h
@@ -12,7 +12,8 @@ class LEVELSNAPSHOTSEDITOR_API ULevelSnapshotsEditorSettings : public UObject
 	GENERATED_BODY()
 public:
 
-	static ULevelSnapshotsEditorSettings* Get();
+	static const ULevelSnapshotsEditorSettings* Get();
+	static ULevelSnapshotsEditorSettings* GetMutable();
 	
 	ULevelSnapshotsEditorSettings(const FObjectInitializer& ObjectInitializer);
 
@@ -85,5 +86,5 @@ public:
 private:
 	
 	/* If the user overrides the Name field in the creation form, the override will be saved here so it can be recalled. */
-	FString LevelSnapshotNameOverride;
+	TAttribute<FString> LevelSnapshotNameOverride;
 };
diff --git a/Engine/Plugins/VirtualProduction/LevelSnapshots/Source/LevelSnapshotsEditor/Private/Util/TakeSnapshotUtil.cpp b/Engine/Plugins/VirtualProduction/LevelSnapshots/Source/LevelSnapshotsEditor/Private/Util/TakeSnapshotUtil.cpp
--- a/Engine/Plugins/VirtualProduction/LevelSnapshots/Source/LevelSnapshotsEditor/Private/Util/TakeSnapshotUtil.cpp
+++ b/Engine/Plugins/VirtualProduction/LevelSnapshots/Source/LevelSnapshotsEditor/Private/Util/TakeSnapshotUtil.cpp
@@ -39,7 +39,7 @@ namespace UE::LevelSnapshots::Editor
 			return;
 		}
 		
-		ULevelSnapshotsEditorSettings* DataManagementSettings = ULevelSnapshotsEditorSettings::Get();
+		ULevelSnapshotsEditorSettings* DataManagementSettings = ULevelSnapshotsEditorSettings::GetMutable();
 		DataManagementSettings->ValidateRootLevelSnapshotSaveDirAsGameContentRelative();
 		DataManagementSettings->SanitizeAllProjectSettingsPaths(true);
 
diff --git a/Engine/Plugins/VirtualProduction/LevelSnapshots/Source/LevelSnapshotsEditor/Private/Widgets/SLevelSnapshotsEditorCreationForm.cpp b/Engine/Plugins/VirtualProduction/LevelSnapshots/Source/LevelSnapshotsEditor/Private/Widgets/SLevelSnapshotsEditorCreationForm.cpp
--- a/Engine/Plugins/VirtualProduction/LevelSnapshots/Source/LevelSnapshotsEditor/Private/Widgets/SLevelSnapshotsEditorCreationForm.cpp
+++ b/Engine/Plugins/VirtualProduction/LevelSnapshots/Source/LevelSnapshotsEditor/Private/Widgets/SLevelSnapshotsEditorCreationForm.cpp
@@ -231,7 +231,7 @@ TSharedRef<SWidget> SLevelSnapshotsEditorCreationForm::MakeDataManagementSetting
 	DetailsViewArgs.bShowScrollBar = false;
 
 	TSharedRef<IDetailsView> Details = PropertyEditorModule.CreateDetailView(DetailsViewArgs);
-	const TArray<UObject*> ProjectSettingsObjects = { ULevelSnapshotsEditorSettings::Get() };
+	const TArray<UObject*> ProjectSettingsObjects = { ULevelSnapshotsEditorSettings::GetMutable() };
 	// By requirement, we're only supposed to show the data management settings
 	Details->RegisterInstancedCustomPropertyLayout(
 		ULevelSnapshotsEditorSettings::StaticClass(),
@@ -250,14 +250,17 @@ SLevelSnapshotsEditorCreationForm::~SLevelSnapshotsEditorCreationForm()
 
 FText SLevelSnapshotsEditorCreationForm::GetNameOverrideText() const
 {
+	const ULevelSnapshotsEditorSettings* Settings = ULevelSnapshotsEditorSettings::Get();
+	check(Settings);
+	
 	UWorld* World = ULevelSnapshotsEditorData::GetEditorWorld();
 	if (!ensure(World))
 	{
-		return FText::FromString(ULevelSnapshotsEditorSettings::Get()->GetNameOverride());
+		return FText::FromString(Settings->GetNameOverride());
 	}
 
 	return ULevelSnapshotsEditorSettings::ParseLevelSnapshotsTokensInText(
-		FText::FromString(ULevelSnapshotsEditorSettings::Get()->GetNameOverride()),
+		FText::FromString(Settings->GetNameOverride()),
 		World->GetName()
 		);
 }
@@ -267,7 +270,7 @@ void SLevelSnapshotsEditorCreationForm::SetNameOverrideText(const FText& InNewTe
 	FString NameAsString = InNewText.ToString();
 	ULevelSnapshotsEditorSettings::SanitizePathInline(NameAsString, true);
 
-	ULevelSnapshotsEditorSettings* Settings = ULevelSnapshotsEditorSettings::Get();
+	ULevelSnapshotsEditorSettings* Settings = ULevelSnapshotsEditorSettings::GetMutable();
 	Settings->SetNameOverride(NameAsString);
 	bNameDiffersFromDefault = Settings->IsNameOverridden();
 }
@@ -301,7 +304,7 @@ FReply SLevelSnapshotsEditorCreationForm::OnCreateButtonPressed()
 void SLevelSnapshotsEditorCreationForm::OnWindowClosed(const TSharedRef<SWindow>& ParentWindow) const
 {
 	const FVector2D WindowSize = ParentWindow->GetClientSizeInScreen();
-	ULevelSnapshotsEditorSettings* Settings = ULevelSnapshotsEditorSettings::Get();
+	ULevelSnapshotsEditorSettings* Settings = ULevelSnapshotsEditorSettings::GetMutable();
 	Settings->SetLastCreationWindowSize(WindowSize);
 	Settings->SaveConfig();
 
