Level Snapshots:

-Fix for: Issue 163135: Compile errors in LevelSnapshotsEditorSettings.h and LevelSnapshotsEditorSettings.cpp

diff --git a/Engine/Plugins/VirtualProduction/LevelSnapshots/Source/LevelSnapshotsEditor/Private/LevelSnapshotsEditorSettings.h b/Engine/Plugins/VirtualProduction/LevelSnapshots/Source/LevelSnapshotsEditor/Private/LevelSnapshotsEditorSettings.h
--- a/Engine/Plugins/VirtualProduction/LevelSnapshots/Source/LevelSnapshotsEditor/Private/LevelSnapshotsEditorSettings.h
+++ b/Engine/Plugins/VirtualProduction/LevelSnapshots/Source/LevelSnapshotsEditor/Private/LevelSnapshotsEditorSettings.h
@@ -2,9 +2,13 @@
 
 #pragma once
 
-#include "CoreMinimal.h"
+#include "Containers/UnrealString.h"
+#include "Internationalization/Text.h"
+#include "Math/Vector2D.h"
 #include "Misc/Attribute.h"
+#include "UObject/SoftObjectPath.h"
 #include "UObject/Object.h"
+
 #include "LevelSnapshotsEditorSettings.generated.h"
 
 UCLASS(config = Engine, defaultconfig)
