Level Snapshots:

-Update all instances of FEditorStyle to FAppStyle

diff --git a/Engine/Plugins/VirtualProduction/LevelSnapshots/Source/LevelSnapshotsEditor/Private/LevelSnapshotsEditorStyle.cpp b/Engine/Plugins/VirtualProduction/LevelSnapshots/Source/LevelSnapshotsEditor/Private/LevelSnapshotsEditorStyle.cpp
--- a/Engine/Plugins/VirtualProduction/LevelSnapshots/Source/LevelSnapshotsEditor/Private/LevelSnapshotsEditorStyle.cpp
+++ b/Engine/Plugins/VirtualProduction/LevelSnapshots/Source/LevelSnapshotsEditor/Private/LevelSnapshotsEditorStyle.cpp
@@ -75,7 +75,10 @@ TSharedRef< FSlateStyleSet > FLevelSnapshotsEditorStyle::Create()
 	Style->Set("LevelSnapshots.ToolbarButton.Small", new IMAGE_BRUSH_SVG("Icons/Icon", Icon16x16));
 	Style->Set("ClassIcon.LevelSnapshotsEditorData", new IMAGE_BRUSH_SVG("Icons/Icon", Icon16x16));
 
-	// Brush
+	// Button Icons
+	Style->Set( "LevelSnapshotsEditor.Add", new IMAGE_BRUSH_SVG("Starship/MainToolbar/ToolBadgePlus", Icon20x20));
+
+	// Border Brushes
 	Style->Set("LevelSnapshotsEditor.GroupBorder", new FSlateRoundedBoxBrush(FStyleColors::Secondary, 4.0f));
 	Style->Set("LevelSnapshotsEditor.BrightBorder", new FSlateColorBrush(FColor(112, 112, 112, 100)));
 	Style->Set("LevelSnapshotsEditor.FilterSelected", new BOX_BRUSH("Common/Filter_Selected", FMargin(18.0f / 64.0f)));
@@ -102,8 +105,10 @@ TSharedRef< FSlateStyleSet > FLevelSnapshotsEditorStyle::Create()
 
 	FTextBlockStyle AndTextStyle = FAppStyle::Get().GetWidgetStyle<FTextBlockStyle>("Graph.CompactNode.Title");
 	FTextBlockStyle OrTextStyle = AndTextStyle;
-	AndTextStyle.SetFont(FCoreStyle::GetDefaultFontStyle("BoldCondensed", 16 ));
-	OrTextStyle.SetFont(FCoreStyle::GetDefaultFontStyle("BoldCondensed", 18 ));
+	FSlateFontInfo AndFont = FAppStyle::Get().GetFontStyle("BoldCondensed"); AndFont.Size = 16;
+	FSlateFontInfo OrFont = FAppStyle::Get().GetFontStyle("BoldCondensed"); OrFont.Size = 18;
+	AndTextStyle.SetFont(AndFont);
+	OrTextStyle.SetFont(OrFont);
 	Style->Set("LevelSnapshotsEditor.FilterRow.And", AndTextStyle);
 	Style->Set("LevelSnapshotsEditor.FilterRow.Or", OrTextStyle);
 	
diff --git a/Engine/Plugins/VirtualProduction/LevelSnapshots/Source/LevelSnapshotsEditor/Private/Views/Filter/SLevelSnapshotsEditorFilters.cpp b/Engine/Plugins/VirtualProduction/LevelSnapshots/Source/LevelSnapshotsEditor/Private/Views/Filter/SLevelSnapshotsEditorFilters.cpp
--- a/Engine/Plugins/VirtualProduction/LevelSnapshots/Source/LevelSnapshotsEditor/Private/Views/Filter/SLevelSnapshotsEditorFilters.cpp
+++ b/Engine/Plugins/VirtualProduction/LevelSnapshots/Source/LevelSnapshotsEditor/Private/Views/Filter/SLevelSnapshotsEditorFilters.cpp
@@ -17,6 +17,7 @@
 #include "IDetailsView.h"
 #include "Modules/ModuleManager.h"
 #include "PropertyEditorModule.h"
+#include "SPositiveActionButton.h"
 #include "Widgets/Input/SButton.h"
 #include "Widgets/Layout/SScrollBox.h"
 #include "Widgets/Notifications/SNotificationList.h"
@@ -84,35 +85,7 @@ SLevelSnapshotsEditorFilters::~SLevelSnapshotsEditorFilters()
 }
 
 void SLevelSnapshotsEditorFilters::Construct(const FArguments& InArgs, ULevelSnapshotsEditorData* EditorData)
-{
-	struct Local
-	{
-		static TSharedRef<SWidget> CreatePlusText(const FText& Text)
-		{
-			return SNew(SHorizontalBox)
-                    + SHorizontalBox::Slot()
-                    .HAlign(HAlign_Center)
-                    .AutoWidth()
-                    .Padding(FMargin(0.f, 1.f))
-                    [
-	                    SNew(STextBlock)
-						.Font(FAppStyle::Get().GetFontStyle("FontAwesome.10"))
-	                    .TextStyle(FAppStyle::Get(), "NormalText.Important")
-	                    .Text(FEditorFontGlyphs::Plus)
-                    ]
-
-                    + SHorizontalBox::Slot()
-                    .HAlign(HAlign_Left)
-                    .AutoWidth()
-                    .Padding(2.f, 1.f)
-                    [
-                        SNew(STextBlock)
-                        .Justification(ETextJustify::Center)
-                        .TextStyle(FAppStyle::Get(), "NormalText.Important")
-                        .Text(Text)
-                    ];
-		}
-	};
+{	
 	
 	EditorDataPtr = EditorData;
 
@@ -218,13 +191,9 @@ void SLevelSnapshotsEditorFilters::Construct(const FArguments& InArgs, ULevelSna
                     .AutoHeight()
                     [
                         SNew(SButton)
-                            .ButtonStyle(FAppStyle::Get(), "RoundButton")
-                            .ContentPadding(FMargin(4.0, 10.0))
-                            .OnClicked(this, &SLevelSnapshotsEditorFilters::AddFilterClick)
-                            .HAlign(HAlign_Center)
-                            [
-								Local::CreatePlusText(LOCTEXT("AddFilterGroup", "Filter Group"))
-                            ]
+                        .Text(LOCTEXT("AddFilterGroup", "Filter Group"))
+                        .OnClicked(this, &SLevelSnapshotsEditorFilters::AddFilterClick)
+						.HAlign(HAlign_Center)
                     ]
                 ]
             ]
diff --git a/Engine/Plugins/VirtualProduction/LevelSnapshots/Source/LevelSnapshotsEditor/Private/Views/Results/SLevelSnapshotsEditorResults.cpp b/Engine/Plugins/VirtualProduction/LevelSnapshots/Source/LevelSnapshotsEditor/Private/Views/Results/SLevelSnapshotsEditorResults.cpp
--- a/Engine/Plugins/VirtualProduction/LevelSnapshots/Source/LevelSnapshotsEditor/Private/Views/Results/SLevelSnapshotsEditorResults.cpp
+++ b/Engine/Plugins/VirtualProduction/LevelSnapshots/Source/LevelSnapshotsEditor/Private/Views/Results/SLevelSnapshotsEditorResults.cpp
@@ -40,6 +40,15 @@ void SLevelSnapshotsEditorResults::Construct(const FArguments& InArgs, ULevelSna
 	OnRefreshResultsHandle = InEditorData->OnRefreshResults.AddSP(this, &SLevelSnapshotsEditorResults::RefreshResults, false);
 	
 	FMenuBuilder ShowOptionsMenuBuilder = BuildShowOptionsMenu();
+	
+	FSlateFontInfo SnapshotNameTextFont = FAppStyle::Get().GetFontStyle("Bold");
+	SnapshotNameTextFont.Size = 16;
+
+	FSlateFontInfo SelectedActorCountTextFont = FAppStyle::Get().GetFontStyle("Bold");
+	SelectedActorCountTextFont.Size = 14;
+	
+	FSlateFontInfo MiscActorCountTextFont = FAppStyle::Get().GetFontStyle("Regular");
+	MiscActorCountTextFont.Size = 10;
 
 	ChildSlot
 	[
@@ -68,7 +77,7 @@ void SLevelSnapshotsEditorResults::Construct(const FArguments& InArgs, ULevelSna
 			[
 				SAssignNew(SelectedSnapshotNamePtr, STextBlock)
 				.Text(DefaultNameText)
-				.Font(FCoreStyle::GetDefaultFontStyle("Bold", 16))
+				.Font(SnapshotNameTextFont)
 			]
 		]	
 		
@@ -175,7 +184,7 @@ void SLevelSnapshotsEditorResults::Construct(const FArguments& InArgs, ULevelSna
 					.AutoHeight()
 					[
 						SAssignNew(SelectedActorCountText, STextBlock)
-						.Font(FCoreStyle::GetDefaultFontStyle("Bold", 14))
+						.Font(SelectedActorCountTextFont)
 						.Justification(ETextJustify::Right)
 					]
 
@@ -183,7 +192,7 @@ void SLevelSnapshotsEditorResults::Construct(const FArguments& InArgs, ULevelSna
 					.AutoHeight()
 					[
 						SAssignNew(MiscActorCountText, STextBlock)
-						.Font(FCoreStyle::GetDefaultFontStyle("Regular", 10))
+						.Font(MiscActorCountTextFont)
 						.Justification(ETextJustify::Right)
 					]
 				]
diff --git a/Engine/Plugins/VirtualProduction/LevelSnapshots/Source/LevelSnapshotsEditor/Private/Views/Results/SLevelSnapshotsEditorResultsRow.cpp b/Engine/Plugins/VirtualProduction/LevelSnapshots/Source/LevelSnapshotsEditor/Private/Views/Results/SLevelSnapshotsEditorResultsRow.cpp
--- a/Engine/Plugins/VirtualProduction/LevelSnapshots/Source/LevelSnapshotsEditor/Private/Views/Results/SLevelSnapshotsEditorResultsRow.cpp
+++ b/Engine/Plugins/VirtualProduction/LevelSnapshots/Source/LevelSnapshotsEditor/Private/Views/Results/SLevelSnapshotsEditorResultsRow.cpp
@@ -241,10 +241,9 @@ void SLevelSnapshotsEditorResultsRow::GenerateAddedAndRemovedRowComponents(
 	.AutoWidth()
 	.Padding(FMargin(5, 0))
 	[
-		SNew(STextBlock)
-		.Font(FAppStyle::Get().GetFontStyle("FontAwesome.10"))
-        .TextStyle(FAppStyle::Get(), "NormalText.Important")
-        .Text(PinnedItem->GetRowType() == FLevelSnapshotsEditorResultsRow::RemovedComponentToAdd ? FEditorFontGlyphs::Plus : FEditorFontGlyphs::Minus)
+		SNew(SImage)
+		.Image(FAppStyle::Get().GetBrush(
+			PinnedItem->GetRowType() == FLevelSnapshotsEditorResultsRow::RemovedComponentToAdd ? "Icons.Plus" : "Icons.Minus"))
 	];
 
 	BasicRowWidgets->AddSlot()
@@ -252,7 +251,7 @@ void SLevelSnapshotsEditorResultsRow::GenerateAddedAndRemovedRowComponents(
 	.HAlign(HAlign_Left)
 	[
 		SNew(STextBlock)
-		.Font(FCoreStyle::GetDefaultFontStyle("Italic", FCoreStyle::RegularTextSize))
+		.Font(FAppStyle::Get().GetFontStyle("Italic"))
 		.Text(InDisplayText)
 		.ToolTipText(
 			FText::Format
diff --git a/Engine/Plugins/VirtualProduction/LevelSnapshots/Source/LevelSnapshotsEditor/Private/Widgets/Filter/SFilterSearchMenu.cpp b/Engine/Plugins/VirtualProduction/LevelSnapshots/Source/LevelSnapshotsEditor/Private/Widgets/Filter/SFilterSearchMenu.cpp
--- a/Engine/Plugins/VirtualProduction/LevelSnapshots/Source/LevelSnapshotsEditor/Private/Widgets/Filter/SFilterSearchMenu.cpp
+++ b/Engine/Plugins/VirtualProduction/LevelSnapshots/Source/LevelSnapshotsEditor/Private/Widgets/Filter/SFilterSearchMenu.cpp
@@ -24,7 +24,7 @@ namespace
 	{
 		class FFilterMenuBuilder
 		{
-			FMenuBuilder MenuBuilder = FMenuBuilder(true, nullptr, nullptr, true, &FCoreStyle::Get(), false);
+			FMenuBuilder MenuBuilder = FMenuBuilder(true, nullptr, nullptr, true, &FAppStyle::Get(), false);
 			UFavoriteFilterContainer* FilterModel;
 			
 			const SFilterSearchMenu::FFilterDelegate OnSelectFilter;
@@ -207,7 +207,7 @@ void SFilterSearchMenu::Construct(const FArguments& InArgs, UFavoriteFilterConta
 		);
 
 
-	FMenuBuilder BaseMenuBuilder(true, nullptr, nullptr, false, &FCoreStyle::Get(), false);
+	FMenuBuilder BaseMenuBuilder(true, nullptr, nullptr, false, &FAppStyle::Get(), false);
 	BaseMenuBuilder.AddWidget(SearchBox.ToSharedRef(), FText(), true, false);
 	BaseMenuBuilder.AddWidget(SAssignNew(MenuContainer, SHorizontalBox), FText(), true, false);
 	ChildSlot
@@ -343,4 +343,4 @@ void SFilterSearchMenu::OnUserConfirmedItemSelection(TSharedPtr<TSubclassOf<ULev
 }
 
 #undef LOCTEXT_NAMESPACE
- 
\ No newline at end of file
+ 
diff --git a/Engine/Plugins/VirtualProduction/LevelSnapshots/Source/LevelSnapshotsEditor/Private/Widgets/SLevelSnapshotsEditorBrowser.cpp b/Engine/Plugins/VirtualProduction/LevelSnapshots/Source/LevelSnapshotsEditor/Private/Widgets/SLevelSnapshotsEditorBrowser.cpp
--- a/Engine/Plugins/VirtualProduction/LevelSnapshots/Source/LevelSnapshotsEditor/Private/Widgets/SLevelSnapshotsEditorBrowser.cpp
+++ b/Engine/Plugins/VirtualProduction/LevelSnapshots/Source/LevelSnapshotsEditor/Private/Widgets/SLevelSnapshotsEditorBrowser.cpp
@@ -13,6 +13,7 @@
 #include "AssetThumbnail.h"
 #include "ContentBrowserDelegates.h"
 #include "ContentBrowserModule.h"
+#include "Editor/EditorStyle/Public/EditorStyleSet.h"
 #include "Framework/MultiBox/MultiBoxBuilder.h"
 #include "IContentBrowserSingleton.h"
 #include "Misc/ScopedSlowTask.h"
@@ -302,6 +303,9 @@ TSharedRef<SToolTip> SLevelSnapshotsEditorBrowser::CreateCustomTooltip(FAssetDat
 	FString OutCaptureTime;
 	AssetData.GetTagValue("CaptureTime",OutCaptureTime);
 	
+	FSlateFontInfo AssetNameTextFont = FAppStyle::Get().GetFontStyle("Bold");
+	AssetNameTextFont.Size = 10;
+	
 	return SNew(SToolTip)
 	[
 		SNew(SBox)
@@ -334,7 +338,7 @@ TSharedRef<SToolTip> SLevelSnapshotsEditorBrowser::CreateCustomTooltip(FAssetDat
 				[
 					SNew(STextBlock)
 					.Text(FText::FromName(AssetData.AssetName))
-					.Font(FCoreStyle::GetDefaultFontStyle("Bold", 10))
+					.Font(AssetNameTextFont)
 					.OverflowPolicy(ETextOverflowPolicy::Ellipsis)
 				]
 
diff --git a/Engine/Plugins/VirtualProduction/LevelSnapshots/Source/LevelSnapshotsEditor/Private/Widgets/SLevelSnapshotsEditorCreationForm.cpp b/Engine/Plugins/VirtualProduction/LevelSnapshots/Source/LevelSnapshotsEditor/Private/Widgets/SLevelSnapshotsEditorCreationForm.cpp
--- a/Engine/Plugins/VirtualProduction/LevelSnapshots/Source/LevelSnapshotsEditor/Private/Widgets/SLevelSnapshotsEditorCreationForm.cpp
+++ b/Engine/Plugins/VirtualProduction/LevelSnapshots/Source/LevelSnapshotsEditor/Private/Widgets/SLevelSnapshotsEditorCreationForm.cpp
@@ -89,6 +89,15 @@ void SLevelSnapshotsEditorCreationForm::Construct(
 	CallOnCloseDelegate = CallOnClose;
 	bNameDiffersFromDefault = ULevelSnapshotsEditorSettings::Get()->IsNameOverridden();
 	
+	FSlateFontInfo NameTextFont = FAppStyle::Get().GetFontStyle("Regular");
+	NameTextFont.Size = 9;
+	
+	FSlateFontInfo OverrideNameTextFont = FAppStyle::Get().GetFontStyle("Bold");
+	OverrideNameTextFont.Size = 16;
+
+	FSlateFontInfo DescriptionTextFont = FAppStyle::Get().GetFontStyle("Bold");
+	DescriptionTextFont.Size = 10;
+	
 	ChildSlot
 	[
 		SNew(SBorder)
@@ -114,7 +123,7 @@ void SLevelSnapshotsEditorCreationForm::Construct(
 					.VAlign(VAlign_Top)
 					[
 						SNew(STextBlock)
-						.Font(FCoreStyle::GetDefaultFontStyle("Regular", 9))
+						.Font(NameTextFont)
 						.ColorAndOpacity(FColor(200, 200, 200))
 						.Text(NSLOCTEXT("LevelSnapshots", "CreationForm_SnapshotNameLabel", "Name"))
 					]
@@ -129,7 +138,7 @@ void SLevelSnapshotsEditorCreationForm::Construct(
 						.HAlign(HAlign_Fill)
 						[
 							SNew(SEditableTextBox)
-							.Font(FCoreStyle::GetDefaultFontStyle("Bold", 16))
+							.Font(OverrideNameTextFont)
 							.BackgroundColor(FLinearColor::Transparent)
 							.ForegroundColor(FSlateColor::UseForeground())
 							.Justification(ETextJustify::Center)
@@ -169,7 +178,7 @@ void SLevelSnapshotsEditorCreationForm::Construct(
 			.VAlign(VAlign_Fill)
 			[
 				SNew(SMultiLineEditableTextBox)
-				.Font(FCoreStyle::GetDefaultFontStyle("Bold", 10))
+				.Font(DescriptionTextFont)
 				.BackgroundColor(FLinearColor(0.2f, 0.2f, 0.2f))
 				.ForegroundColor(FSlateColor::UseForeground())
 				.SelectAllTextWhenFocused(true)
