Fix crash due to keeping reference to reallocated memory

diff --git a/Engine/Plugins/VirtualProduction/LevelSnapshots/Source/LevelSnapshots/Private/Data/Util/Restoration/AllocationActorUtil.cpp b/Engine/Plugins/VirtualProduction/LevelSnapshots/Source/LevelSnapshots/Private/Data/Util/Restoration/AllocationActorUtil.cpp
--- a/Engine/Plugins/VirtualProduction/LevelSnapshots/Source/LevelSnapshots/Private/Data/Util/Restoration/AllocationActorUtil.cpp
+++ b/Engine/Plugins/VirtualProduction/LevelSnapshots/Source/LevelSnapshots/Private/Data/Util/Restoration/AllocationActorUtil.cpp
@@ -192,19 +192,23 @@ namespace UE::LevelSnapshots::Private::Internal
 		SCOPED_SNAPSHOT_CORE_TRACE(GetDeserialized);
 		UE_LOG(LogLevelSnapshots, Verbose, TEXT("========== Get Deserialized %s =========="), *OriginalActorPath.ToString());
 		
-		FActorSnapshotCache& ActorCache = Cache.ActorCache.FindOrAdd(OriginalActorPath);
 		const TOptional<TNonNullPtr<AActor>> Preallocated = GetPreallocated(OriginalActorPath, WorldData, Cache, SnapshotWorld);
 		if (!Preallocated)
 		{
 			return {};
 		}
+		
+		FActorSnapshotCache& ActorCache = Cache.ActorCache.FindOrAdd(OriginalActorPath);
 		ActorCache.bReceivedSerialisation = true;
 		
 		FActorSnapshotData& ActorData = WorldData.ActorData[OriginalActorPath];
-		const auto ProcessObjectDependency = [&ActorCache](int32 OriginalObjectDependency)
+		const auto ProcessObjectDependency = [&Cache, &OriginalActorPath](int32 OriginalObjectDependency)
 		{
+			// This look-up must be done every time because FSnapshotDataCache::ActorCache may have been reallocated 
+			FActorSnapshotCache& ActorCache = Cache.ActorCache.FindOrAdd(OriginalActorPath);
 			ActorCache.ObjectDependencies.Add(OriginalObjectDependency);
 		};
+		
 		AActor* PreallocatedActor = Preallocated.GetValue();
 		{
 			const FRestoreObjectScope FinishRestore = PreActorRestore_SnapshotWorld(PreallocatedActor, ActorData.CustomActorSerializationData, WorldData, Cache, ProcessObjectDependency, InLocalisationSnapshotPackage);
@@ -213,7 +217,7 @@ namespace UE::LevelSnapshots::Private::Internal
 			UE_LOG(LogLevelSnapshots, Verbose, TEXT("ActorLabel is \"%s\" for \"%s\" (editor object path \"%s\")"), *PreallocatedActor->GetActorLabel(), *PreallocatedActor->GetPathName(), *OriginalActorPath.ToString());
 #endif
 		}
-
+		
 		DeserializeComponents(ActorData, PreallocatedActor, WorldData,
 			[&WorldData, &ProcessObjectDependency, &Cache, InLocalisationSnapshotPackage](
 				FSubobjectSnapshotData& SerializedCompData,
