Fixing non unity build

diff --git a/Engine/Plugins/VirtualProduction/LevelSnapshots/Source/LevelSnapshotsEditor/Private/Data/FavoriteFilterContainer.cpp b/Engine/Plugins/VirtualProduction/LevelSnapshots/Source/LevelSnapshotsEditor/Private/Data/FavoriteFilterContainer.cpp
--- a/Engine/Plugins/VirtualProduction/LevelSnapshots/Source/LevelSnapshotsEditor/Private/Data/FavoriteFilterContainer.cpp
+++ b/Engine/Plugins/VirtualProduction/LevelSnapshots/Source/LevelSnapshotsEditor/Private/Data/FavoriteFilterContainer.cpp
@@ -83,7 +83,7 @@ namespace
 		IAssetRegistry& AssetRegistry = AssetRegistryModule.Get();
 		
 		TArray<FAssetData> BlueprintList;
-		AssetRegistry.GetAssetsByClass(UBlueprint::StaticClass()->GetFName(), BlueprintList);
+		AssetRegistry.GetAssetsByClass(UBlueprint::StaticClass()->GetClassPathName(), BlueprintList);
 
 		// Returns = 
 		TArray<TSubclassOf<ULevelSnapshotBlueprintFilter>> Result;
diff --git a/Engine/Source/Runtime/CoreUObject/Public/Misc/AssetRegistryInterface.h b/Engine/Source/Runtime/CoreUObject/Public/Misc/AssetRegistryInterface.h
--- a/Engine/Source/Runtime/CoreUObject/Public/Misc/AssetRegistryInterface.h
+++ b/Engine/Source/Runtime/CoreUObject/Public/Misc/AssetRegistryInterface.h
@@ -7,6 +7,7 @@
 #include "Containers/ContainersFwd.h"
 #include "Modules/ModuleInterface.h"
 #include "Modules/ModuleManager.h"
+#include "UObject/TopLevelAssetPath.h"
 
 namespace EAssetRegistryDependencyType
 {
