Merge changes from Otter stream to Main

-------------------------------------------

Fix undefined behaviour due to keeping a reference to memory that is sometimes reallocated. Also fix several warnings, logs, and failing ensures.

#[fyi] Jason.Walter
Rename Data/Util/Restoration folder to Data/Util/WorldData and move cpp and h files into Data/Util/WorldData folder

#[fyi] Jason.Walter
Fix lag when hovering snapshot

#[fyi] Jason.Walter
Better way to obtain stats for applying snapshots:
- Make apply snapshot dialog show current actor.
- Add console variables for logging how long it took to recreate and serialize actors

#[fyi] Jason.Walter
Fix recreated actors having properties left to restore when they were different from the CDO

#[fyi] Jason.Walter
Save class archetypes and fix hashes not getting computed when diffing

#[fyi] Jason.Walter
Snapshots now use Oodle compression when saved. Significantly reduces size and time taken to save & load snapshots.

#[fyi] Jason.Walter
Added missing files

#[fyi] Jason.Walter
Improve compression defaults

#[fyi] Jason.Walter
Refactor and minor clean-up of diffing and filtering the maps

#[fyi] Jason.Walter
Compute actor hash before comparing properties

#[fyi] Jason.Walter
Remove unused code

#[fyi] Jason.Walter
Make tooltip's of snapshot object rows show the object path

#[fyi] Jason.Walter
Fix added and deleted actors still getting restored even though unchecked

#[fyi] Jason.Walter
ActorInMapFilter now also checks actor that are supposed to get added to the level

#[fyi] Jason.Walter
Fix actor rows getting created multiple times

#[fyi] Jason.Walter
Fix construction script controlled properties not restoring correctly

#[fyi] Jason.Walter
Future proof compressed saving of FWorldSnapshotData

#[fyi] Jason.Walter
Transient snapshot world now also creates actors in corresponding transient sublevels

#[fyi] Jason.Walter
Fix restoring invalid references with Level Snapshots: actors from level Foo can no longer be restored to references actors from level Bar

#[fyi] Jason.Walter
Group actors are now recreated and deleted correctly

#[fyi] Jason.Walter

diff --git a/Engine/Plugins/VirtualProduction/LevelSnapshots/Source/LevelShapshotFilters/Private/Builtin/ActorInMapFilter.cpp b/Engine/Plugins/VirtualProduction/LevelSnapshots/Source/LevelShapshotFilters/Private/Builtin/ActorInMapFilter.cpp
--- a/Engine/Plugins/VirtualProduction/LevelSnapshots/Source/LevelShapshotFilters/Private/Builtin/ActorInMapFilter.cpp
+++ b/Engine/Plugins/VirtualProduction/LevelSnapshots/Source/LevelShapshotFilters/Private/Builtin/ActorInMapFilter.cpp
@@ -6,31 +6,18 @@
 #include "Engine/World.h"
 #include "GameFramework/Actor.h"
 
-namespace ActorInMapFilter
+namespace UE::LevelSnapshotsFilters::Private
 {
-	static bool IsActorInMap(const AActor* Actor, const FString& MapNameToCheck)
+	static bool IsActorInMap(const FSoftObjectPath& Actor, const FSoftObjectPath& MapNameToCheck)
 	{
-		const ULevel* Level = Actor->GetLevel();
-
-		if (!ensure(Level))
-		{
-			return false;
-		}
-
-		const UWorld* World = Level->GetTypedOuter<UWorld>();
-		if (!ensure(World))
-		{
-			return false;
-		}
-
-		return World->GetName() == MapNameToCheck;
+		return Actor.GetAssetPathName() == MapNameToCheck.GetAssetPathName();
 	}
 
-	static EFilterResult::Type IsActorAllowed(const AActor* Actor, const TArray<TSoftObjectPtr<UWorld>>& AllowedLevels)
+	static EFilterResult::Type IsActorAllowed(const FSoftObjectPath& Actor, const TArray<TSoftObjectPtr<UWorld>>& AllowedLevels)
 	{
 		for (const TSoftObjectPtr<UWorld>& AllowedLevel : AllowedLevels)
 		{
-			if (IsActorInMap(Actor, AllowedLevel.GetAssetName()))
+			if (IsActorInMap(Actor, AllowedLevel.ToSoftObjectPath()))
 			{
 				return EFilterResult::Include;
 			}
@@ -42,12 +29,17 @@ namespace ActorInMapFilter
 
 EFilterResult::Type UActorInMapFilter::IsActorValid(const FIsActorValidParams& Params) const
 {
-	return ActorInMapFilter::IsActorAllowed(Params.LevelActor, AllowedLevels);
+	return UE::LevelSnapshotsFilters::Private::IsActorAllowed(Params.LevelActor, AllowedLevels);
 }
 
 EFilterResult::Type UActorInMapFilter::IsAddedActorValid(const FIsAddedActorValidParams& Params) const
 {
-	return ActorInMapFilter::IsActorAllowed(Params.NewActor, AllowedLevels);
+	return UE::LevelSnapshotsFilters::Private::IsActorAllowed(Params.NewActor, AllowedLevels);
+}
+
+EFilterResult::Type UActorInMapFilter::IsDeletedActorValid(const FIsDeletedActorValidParams& Params) const
+{
+	return UE::LevelSnapshotsFilters::Private::IsActorAllowed(Params.SavedActorPath, AllowedLevels);
 }
 
 
diff --git a/Engine/Plugins/VirtualProduction/LevelSnapshots/Source/LevelShapshotFilters/Public/Builtin/ActorInMapFilter.h b/Engine/Plugins/VirtualProduction/LevelSnapshots/Source/LevelShapshotFilters/Public/Builtin/ActorInMapFilter.h
--- a/Engine/Plugins/VirtualProduction/LevelSnapshots/Source/LevelShapshotFilters/Public/Builtin/ActorInMapFilter.h
+++ b/Engine/Plugins/VirtualProduction/LevelSnapshots/Source/LevelShapshotFilters/Public/Builtin/ActorInMapFilter.h
@@ -17,6 +17,7 @@ public:
 	//~ Begin ULevelSnapshotFilter Interface
 	virtual EFilterResult::Type IsActorValid(const FIsActorValidParams& Params) const override;
 	virtual EFilterResult::Type IsAddedActorValid(const FIsAddedActorValidParams& Params) const override;
+	virtual EFilterResult::Type IsDeletedActorValid(const FIsDeletedActorValidParams& Params) const override;
 	//~ End ULevelSnapshotFilter Interface
 
 private:
diff --git a/Engine/Plugins/VirtualProduction/LevelSnapshots/Source/LevelSnapshots/Private/Archive/ApplySnapshotToEditorArchive.cpp b/Engine/Plugins/VirtualProduction/LevelSnapshots/Source/LevelSnapshots/Private/Archive/ApplySnapshotToEditorArchive.cpp
--- a/Engine/Plugins/VirtualProduction/LevelSnapshots/Source/LevelSnapshots/Private/Archive/ApplySnapshotToEditorArchive.cpp
+++ b/Engine/Plugins/VirtualProduction/LevelSnapshots/Source/LevelSnapshots/Private/Archive/ApplySnapshotToEditorArchive.cpp
@@ -4,63 +4,46 @@
 
 #include "Archive/ClassDefaults/ApplyClassDefaulDataArchive.h"
 #include "Data/RestorationEvents/ApplySnapshotPropertiesScope.h"
+#include "Data/Util/WorldData/ClassDataUtil.h"
 #include "Data/Util/Property/PropertyUtil.h"
-#include "Data/Util/SnapshotObjectUtil.h"
+#include "Data/Util/WorldData/SnapshotObjectUtil.h"
 #include "Data/WorldSnapshotData.h"
 #include "LevelSnapshotsLog.h"
 #include "Selection/PropertySelection.h"
 
-#include "Components/ActorComponent.h"
 #include "Engine/World.h"
-#include "GameFramework/Actor.h"
 #include "Serialization/ObjectReader.h"
 #include "Serialization/ObjectWriter.h"
 #include "UObject/TextProperty.h"
 #include "UObject/UnrealType.h"
-#include "Util/Restoration/WorldDataUtil.h"
+#include "Util/Property/WorldReferenceCheckingUtil.h"
 
 namespace UE::LevelSnapshots::Private::Internal
 {
+	using FShouldSerialize = TFunctionRef<bool(const FArchiveSerializedPropertyChain*, const FProperty*)>;
 	class FCopyProperties : public FObjectWriter
 	{
 		using Super = FObjectWriter;
 
-		const FPropertySelection& PropertiesToSerialize;
+		FShouldSerialize PropertiesToSerialize;
 		UObject* SnapshotObject;
 
-		bool IsWorldObjectProperty(const FProperty* InProperty) const
+		bool ContainsWorldObjectProperty(const FProperty* InProperty) const
 		{
-			const FObjectPropertyBase* ObjectProperty = CastField<FObjectPropertyBase>(InProperty);
-			if (!ObjectProperty)
+			return UE::LevelSnapshots::Private::ContainsSubobjectOrSatisfiesPredicate(SnapshotObject, GetSerializedPropertyChain(), InProperty, [](UObject* ContainedUObject)
 			{
-				return false;
-			}
-
-			const bool bIsMarkedAsSubobject = InProperty->HasAnyPropertyFlags(CPF_InstancedReference | CPF_ContainsInstancedReference | CPF_PersistentInstance);
-			const bool bIsActorOrComponentPtr = ObjectProperty->PropertyClass->IsChildOf(AActor::StaticClass()) || ObjectProperty->PropertyClass->IsChildOf(UActorComponent::StaticClass());
-			if (bIsMarkedAsSubobject || bIsActorOrComponentPtr)
-			{
-				return true;
-			}
-
-			const FArchiveSerializedPropertyChain* PropertyChain = GetSerializedPropertyChain();
-			const bool bContainsWorldReference = UE::LevelSnapshots::Private::FollowPropertyChainUntilPredicateIsTrue(SnapshotObject, PropertyChain, InProperty, [this, ObjectProperty](void* LeafValuePtr)
-			{
-				const UObject* ContainedPtr = ObjectProperty->GetObjectPropertyValue(LeafValuePtr);
-				return LeafValuePtr ? ContainedPtr->IsInA(UWorld::StaticClass()) : false;
+				return ContainedUObject && ContainedUObject->IsInA(UWorld::StaticClass());
 			});
-			return bContainsWorldReference;
 		}
 				
 	public:
 
 		mutable TSet<const FTextProperty*> TextProperties;
 
-		FCopyProperties(UObject* SnapshotObject, TArray<uint8>& SaveLocation, const FPropertySelection& PropertiesToSerialize)
-			:
-			Super(SaveLocation),
-			PropertiesToSerialize(PropertiesToSerialize),
-			SnapshotObject(SnapshotObject)
+		FCopyProperties(UObject* SnapshotObject, TArray<uint8>& SaveLocation, FShouldSerialize PropertiesToSerialize)
+			: Super(SaveLocation)
+			, PropertiesToSerialize(PropertiesToSerialize)
+			, SnapshotObject(SnapshotObject)
 		{
 			ArNoDelta = true;
 		}
@@ -69,12 +52,12 @@ namespace UE::LevelSnapshots::Private::Internal
 		{
 			// Do not copy object reference properties that have a world as outer: They will not be valid when copied from snapshot to editor world.
 			// Hence, we only allow object references to external assets, e.g. UMaterials or UDataAssets
-			if (IsWorldObjectProperty(InProperty))
+			if (ContainsWorldObjectProperty(InProperty))
 			{
 				return true;
 			}
 
-			const bool bIsPropertyAllowed = PropertiesToSerialize.ShouldSerializeProperty(GetSerializedPropertyChain(), InProperty);
+			const bool bIsPropertyAllowed = PropertiesToSerialize(GetSerializedPropertyChain(), InProperty);
 			if (bIsPropertyAllowed)
 			{
 				if (const FTextProperty* TextProperty = CastField<FTextProperty>(InProperty))
@@ -107,37 +90,63 @@ namespace UE::LevelSnapshots::Private::Internal
 			return true;
 		}
 	};
+	
+	static TSet<const FTextProperty*> CopyPastePropertiesDifferentInCDO(FShouldSerialize ShouldSerialize, UObject* InOriginalObject, UObject* InDeserializedVersion)
+	{
+		TArray<uint8> CopiedPropertyData;
+		FCopyProperties CopySimpleProperties(InDeserializedVersion, CopiedPropertyData, ShouldSerialize);
+		InDeserializedVersion->Serialize(CopySimpleProperties);
+		
+		FObjectReader PasteSimpleProperties(InOriginalObject, CopiedPropertyData);
+		return CopySimpleProperties.TextProperties;
+	}
 
 	static TSet<const FTextProperty*> CopyPastePropertiesDifferentInCDO(const FPropertySelection& PropertiesLeftToSerialise, UObject* InOriginalObject, UObject* InDeserializedVersion)
 	{
-		if (!PropertiesLeftToSerialise.IsEmpty())
+		if (PropertiesLeftToSerialise.IsEmpty())
 		{
-			TArray<uint8> CopiedPropertyData;
-			FCopyProperties CopySimpleProperties(InDeserializedVersion, CopiedPropertyData, PropertiesLeftToSerialise);
-			InDeserializedVersion->Serialize(CopySimpleProperties);
-		
-			FObjectReader PasteSimpleProperties(InOriginalObject, CopiedPropertyData);
-			return CopySimpleProperties.TextProperties;
+			return {};
 		}
-		return {};
+
+		return CopyPastePropertiesDifferentInCDO(
+			[&PropertiesLeftToSerialise](const FArchiveSerializedPropertyChain* Chain, const FProperty* Property)
+			{
+				return PropertiesLeftToSerialise.ShouldSerializeProperty(Chain, Property);
+			},
+			InOriginalObject,
+			InDeserializedVersion
+			);
 	}
 
-	static void FixUpTextPropertiesDifferentInCDO(TSet<const FTextProperty*> TextProperties, FWorldSnapshotData& InSharedData, UObject* InOriginalObject)
+	static void FixUpTextPropertiesDifferentInCDO(
+		TSet<const FTextProperty*> TextProperties,
+		FWorldSnapshotData& InSharedData,
+		TOptional<FClassDataIndex> ClassIndex,
+		FSnapshotDataCache& Cache,
+		UObject* InOriginalObject)
 	{
 		const bool bAreTextPropertiesLeft = TextProperties.Num() > 0;
-		if (bAreTextPropertiesLeft)
+		if (!bAreTextPropertiesLeft)
 		{
-			FObjectSnapshotData* ClassDefaults = GetSerializedClassDefaults(InSharedData, InOriginalObject->GetClass()).Get(nullptr);
+			return;
+		}
+
+		if (ClassIndex)
+		{
+			const FSubobjectArchetypeFallbackInfo FallbackInfo { InOriginalObject->GetOuter() };
+			TOptional<TNonNullPtr<FClassSnapshotData>> ClassDefaults =
+				InOriginalObject->IsA<AActor>()
+					? TOptional<TNonNullPtr<FClassSnapshotData>>(&InSharedData.ClassData[*ClassIndex])
+					: GetSubobjectArchetypeData(InSharedData, *ClassIndex, Cache, FallbackInfo);
 			if (ClassDefaults)
 			{
-				FSerializeTextProperties SerializeTextProperties(TextProperties, *ClassDefaults, InSharedData, InOriginalObject);
+				FSerializeTextProperties SerializeTextProperties(TextProperties, *ClassDefaults.GetValue(), InSharedData, InOriginalObject);
 				InOriginalObject->Serialize(SerializeTextProperties);
-			}
-			else
-			{
-				UE_LOG(LogLevelSnapshots, Warning, TEXT("%d text properties have changed in class defaults since snapshot was taken but cannot be restored."), TextProperties.Num());
+				return;
 			}
 		}
+			
+		UE_LOG(LogLevelSnapshots, Warning, TEXT("%d text properties have changed in class defaults since snapshot was taken but cannot be restored."), TextProperties.Num());
 	}
 }
 
@@ -147,7 +156,8 @@ void UE::LevelSnapshots::Private::FApplySnapshotToEditorArchive::ApplyToExisting
 	FSnapshotDataCache& Cache,
 	UObject* InOriginalObject,
 	UObject* InDeserializedVersion,
-	const FPropertySelectionMap& InSelectionMapForResolvingSubobjects)
+	const FPropertySelectionMap& InSelectionMapForResolvingSubobjects,
+	TOptional<FClassDataIndex> ClassIndex)
 {
 	const FPropertySelection* Selection = InSelectionMapForResolvingSubobjects.GetObjectSelection(InOriginalObject).GetPropertySelection();
 	if (Selection && Selection->IsEmpty())
@@ -160,19 +170,22 @@ void UE::LevelSnapshots::Private::FApplySnapshotToEditorArchive::ApplyToExisting
 #if WITH_EDITOR
 	InOriginalObject->Modify();
 #endif
+
+	FPropertySelection PropertiesLeftToSerialize = *Selection;
+	auto TrackSerializedProperties = [&PropertiesLeftToSerialize](const FArchiveSerializedPropertyChain* Chain, const FProperty* Property) { PropertiesLeftToSerialize.RemoveProperty(Chain, Property); };
 	
 	// Step 1: Serialise  properties that were different from CDO at time of snapshotting and that are still different from CDO
-	FApplySnapshotToEditorArchive ApplySavedData(InObjectData, InSharedData, InOriginalObject, InSelectionMapForResolvingSubobjects, Selection, Cache);
+	FApplySnapshotToEditorArchive ApplySavedData(InObjectData, InSharedData, InOriginalObject, InSelectionMapForResolvingSubobjects, Selection, TrackSerializedProperties, Cache);
 	InOriginalObject->Serialize(ApplySavedData);
 	
 	// Step 2: Serialise any remaining properties that were not covered: properties that were equal to the CDO value when the snapshot was taken but now are different from the CDO.
-	TSet<const FTextProperty*> TextProperties = Internal::CopyPastePropertiesDifferentInCDO(ApplySavedData.PropertiesLeftToSerialize, InOriginalObject, InDeserializedVersion);
+	TSet<const FTextProperty*> TextProperties = Internal::CopyPastePropertiesDifferentInCDO(PropertiesLeftToSerialize, InOriginalObject, InDeserializedVersion);
 
 	// Step 3: Serialize FText properties that have changed in CDO since snapshot was taken. 
-	Internal::FixUpTextPropertiesDifferentInCDO(MoveTemp(TextProperties), InSharedData, InOriginalObject);
+	Internal::FixUpTextPropertiesDifferentInCDO(MoveTemp(TextProperties), InSharedData, ClassIndex, Cache, InOriginalObject);
 }
 
-void UE::LevelSnapshots::Private::FApplySnapshotToEditorArchive::ApplyToRecreatedEditorWorldObject(
+void UE::LevelSnapshots::Private::FApplySnapshotToEditorArchive::ApplyToEditorWorldObjectRecreatedWithArchetype(
 	FObjectSnapshotData& InObjectData,
 	FWorldSnapshotData& InSharedData,
 	FSnapshotDataCache& Cache,
@@ -184,8 +197,37 @@ void UE::LevelSnapshots::Private::FApplySnapshotToEditorArchive::ApplyToRecreate
 	
 	// Apply all properties that we saved into the target actor.
 	// We assume that InOriginalObject was already created with the snapshot CDO as template: we do not need Step 2 from ApplyToExistingWorldObject.
-	FApplySnapshotToEditorArchive ApplySavedData(InObjectData, InSharedData, InOriginalObject, InSelectionMapForResolvingSubobjects, {}, Cache);
+	auto Dummy = [](auto,auto){};
+	FApplySnapshotToEditorArchive ApplySavedData(InObjectData, InSharedData, InOriginalObject, InSelectionMapForResolvingSubobjects, {}, Dummy, Cache);
+	InOriginalObject->Serialize(ApplySavedData);
+}
+
+void UE::LevelSnapshots::Private::FApplySnapshotToEditorArchive::ApplyToEditorWorldObjectRecreatedWithoutArchetype(
+	FObjectSnapshotData& InObjectData,
+	FWorldSnapshotData& InSharedData,
+	FSnapshotDataCache& Cache,
+	UObject* InOriginalObject,
+	UObject* InDeserializedVersion,
+	const FPropertySelectionMap& InSelectionMapForResolvingSubobjects,
+	FClassDataIndex ClassIndex)
+{
+	FPropertySelection PropertiesThatWereSerialized;
+	auto TrackSerializedProperties = [&PropertiesThatWereSerialized](const FArchiveSerializedPropertyChain* Chain, const FProperty* Property) { PropertiesThatWereSerialized.AddProperty({ Chain, Property}); };
+	
+	// Step 1: Serialise  properties that were different from CDO at time of snapshotting and that are still different from CDO
+	FApplySnapshotToEditorArchive ApplySavedData(InObjectData, InSharedData, InOriginalObject, InSelectionMapForResolvingSubobjects, {}, TrackSerializedProperties, Cache);
 	InOriginalObject->Serialize(ApplySavedData);
+	
+	// Step 2: Serialise any remaining properties that were not covered: properties that were equal to the CDO value when the snapshot was taken but now are different from the CDO.
+	TSet<const FTextProperty*> TextProperties = Internal::CopyPastePropertiesDifferentInCDO(
+		[&PropertiesThatWereSerialized](const FArchiveSerializedPropertyChain* Chain, const FProperty* Property)
+		{
+			const bool bWasSerializedAlready = PropertiesThatWereSerialized.ShouldSerializeProperty(Chain, Property);
+			return !bWasSerializedAlready;
+		}, InOriginalObject, InDeserializedVersion);
+
+	// Step 3: Serialize FText properties that have changed in CDO since snapshot was taken. 
+	Internal::FixUpTextPropertiesDifferentInCDO(MoveTemp(TextProperties), InSharedData, ClassIndex, Cache, InOriginalObject);
 }
 
 bool UE::LevelSnapshots::Private::FApplySnapshotToEditorArchive::ShouldSkipProperty(const FProperty* InProperty) const
@@ -206,7 +248,7 @@ bool UE::LevelSnapshots::Private::FApplySnapshotToEditorArchive::ShouldSkipPrope
 void UE::LevelSnapshots::Private::FApplySnapshotToEditorArchive::PushSerializedProperty(FProperty* InProperty, const bool bIsEditorOnlyProperty)
 {
 	// Do before call to super because super appends InProperty to property chain
-	PropertiesLeftToSerialize.RemoveProperty(GetSerializedPropertyChain(), InProperty);
+	OnPropertySerialized(GetSerializedPropertyChain(), InProperty);
 	
 	Super::PushSerializedProperty(InProperty, bIsEditorOnlyProperty);
 }
@@ -217,7 +259,7 @@ UObject* UE::LevelSnapshots::Private::FApplySnapshotToEditorArchive::ResolveObje
 #if USE_STABLE_LOCALIZATION_KEYS
 	LocalizationNamespace = GetLocalizationNamespace();
 #endif
-	return UE::LevelSnapshots::Private::ResolveObjectDependencyForEditorWorld(GetSharedData(), Cache, ObjectIndex, LocalizationNamespace, SelectionMapForResolvingSubobjects);
+	return ResolveObjectDependencyForEditorWorld(GetSharedData(), Cache, ObjectIndex, LocalizationNamespace, SelectionMapForResolvingSubobjects);
 }
 
 UE::LevelSnapshots::Private::FApplySnapshotToEditorArchive::FApplySnapshotToEditorArchive(
@@ -226,17 +268,14 @@ UE::LevelSnapshots::Private::FApplySnapshotToEditorArchive::FApplySnapshotToEdit
 	UObject* InOriginalObject,
 	const FPropertySelectionMap& InSelectionMapForResolvingSubobjects,
 	TOptional<const FPropertySelection*> InSelectionSet,
+	FOnSerializeProperty InOnPropertySerialized,
 	FSnapshotDataCache& Cache)
         : Super(InObjectData, InSharedData, true, InOriginalObject)
 		, SelectionMapForResolvingSubobjects(InSelectionMapForResolvingSubobjects)
 		, SelectionSet(InSelectionSet)
+		, OnPropertySerialized(InOnPropertySerialized)
 		, Cache(Cache)
-{
-	if (SelectionSet)
-	{
-		PropertiesLeftToSerialize = *SelectionSet.GetValue();
-	}
-}
+{}
 
 bool UE::LevelSnapshots::Private::FApplySnapshotToEditorArchive::ShouldSerializeAllProperties() const
 {
diff --git a/Engine/Plugins/VirtualProduction/LevelSnapshots/Source/LevelSnapshots/Private/Archive/ApplySnapshotToEditorArchive.h b/Engine/Plugins/VirtualProduction/LevelSnapshots/Source/LevelSnapshots/Private/Archive/ApplySnapshotToEditorArchive.h
--- a/Engine/Plugins/VirtualProduction/LevelSnapshots/Source/LevelSnapshots/Private/Archive/ApplySnapshotToEditorArchive.h
+++ b/Engine/Plugins/VirtualProduction/LevelSnapshots/Source/LevelSnapshots/Private/Archive/ApplySnapshotToEditorArchive.h
@@ -6,6 +6,7 @@
 #include "Selection/PropertySelection.h"
 #include "Selection/PropertySelectionMap.h"
 #include "SnapshotArchive.h"
+#include "SnapshotUtilTypes.h"
 
 struct FSnapshotDataCache;
 
@@ -16,9 +17,46 @@ namespace UE::LevelSnapshots::Private
 	{
 		using Super = FSnapshotArchive;
 	public:
-	
-		static void ApplyToExistingEditorWorldObject(FObjectSnapshotData& InObjectData, FWorldSnapshotData& InSharedData, FSnapshotDataCache& Cache, UObject* InOriginalObject, UObject* InDeserializedVersion, const FPropertySelectionMap& InSelectionMapForResolvingSubobjects);
-		static void ApplyToRecreatedEditorWorldObject(FObjectSnapshotData& InObjectData, FWorldSnapshotData& InSharedData, FSnapshotDataCache& Cache, UObject* InOriginalObject, const FPropertySelectionMap& InSelectionMapForResolvingSubobjects);
+
+		/**
+		 * Serializes selected snapshot data into an object that existed in the world.
+		 * The property selection map contains a set of properties that are supposed to be serialized.
+		 */
+		static void ApplyToExistingEditorWorldObject(
+			FObjectSnapshotData& InObjectData,
+			FWorldSnapshotData& InSharedData,
+			FSnapshotDataCache& Cache,
+			UObject* InOriginalObject,
+			UObject* InDeserializedVersion,
+			const FPropertySelectionMap& InSelectionMapForResolvingSubobjects,
+			TOptional<FClassDataIndex> ClassIndex = {}
+			);
+
+		/**
+		 * Serializes all snapshot data into an object that was recreated and was templated using the snapshot saved archetype / CDO.
+		 * Because the object was recreated from an archetype, we can skip applying CDO's data into the object.
+		 */
+		static void ApplyToEditorWorldObjectRecreatedWithArchetype(
+			FObjectSnapshotData& InObjectData,
+			FWorldSnapshotData& InSharedData,
+			FSnapshotDataCache& Cache,
+			UObject* InOriginalObject,
+			const FPropertySelectionMap& InSelectionMapForResolvingSubobjects
+			);
+
+		/**
+		 * Serializes all snapshot data into an object that was recreated and was templated using the snapshot saved archetype / CDO.
+		 * Because the object was not recreated from any archetype, we must also apply the saved CDO's data to the object.
+		 */
+		static void ApplyToEditorWorldObjectRecreatedWithoutArchetype(
+			FObjectSnapshotData& InObjectData,
+			FWorldSnapshotData& InSharedData,
+			FSnapshotDataCache& Cache,
+			UObject* InOriginalObject,
+			UObject* InDeserializedVersion,
+			const FPropertySelectionMap& InSelectionMapForResolvingSubobjects,
+			FClassDataIndex ClassIndex
+			);
 	
 		//~ Begin FTakeSnapshotArchiveV2 Interface
 		virtual bool ShouldSkipProperty(const FProperty* InProperty) const override;
@@ -33,19 +71,28 @@ namespace UE::LevelSnapshots::Private
 
 	private:
 
-		FApplySnapshotToEditorArchive(FObjectSnapshotData& InObjectData, FWorldSnapshotData& InSharedData, UObject* InOriginalObject, const FPropertySelectionMap& InSelectionMapForResolvingSubobjects, TOptional<const FPropertySelection*> InSelectionSet, FSnapshotDataCache& Cache);
+		using FOnSerializeProperty = TFunctionRef<void(const FArchiveSerializedPropertyChain*, const FProperty*)>;
+
+		FApplySnapshotToEditorArchive(
+			FObjectSnapshotData& InObjectData,
+			FWorldSnapshotData& InSharedData,
+			UObject* InOriginalObject,
+			const FPropertySelectionMap& InSelectionMapForResolvingSubobjects,
+			TOptional<const FPropertySelection*> InSelectionSet,
+			FOnSerializeProperty InOnPropertySerialized,
+			FSnapshotDataCache& Cache
+			);
 
 		bool ShouldSerializeAllProperties() const;
 
-		/* Needed so subobjects can be fully reconstructed */
+		/** Needed so subobjects can be fully reconstructed */
 		const FPropertySelectionMap& SelectionMapForResolvingSubobjects; 
 	
-		/* Immutable list of properties we are supposed to serialise */
+		/** Immutable list of properties we are supposed to serialise. Serialize all properties if unset. */
 		TOptional<const FPropertySelection*> SelectionSet;
-		/* A property is removed when we serialise it.
-		* After serialisation is done, this list tells us which properties are still left to serialise.
-		*/
-		mutable FPropertySelection PropertiesLeftToSerialize;
+
+		/** After serialisation is done, we can tell which properties are still left to serialise (i.e. copy from deserialized). */
+		FOnSerializeProperty OnPropertySerialized;
 
 		FSnapshotDataCache& Cache;
 	};
diff --git a/Engine/Plugins/VirtualProduction/LevelSnapshots/Source/LevelSnapshots/Private/Archive/ClassDefaults/BaseClassDefaultArchive.cpp b/Engine/Plugins/VirtualProduction/LevelSnapshots/Source/LevelSnapshots/Private/Archive/ClassDefaults/BaseClassDefaultArchive.cpp
--- a/Engine/Plugins/VirtualProduction/LevelSnapshots/Source/LevelSnapshots/Private/Archive/ClassDefaults/BaseClassDefaultArchive.cpp
+++ b/Engine/Plugins/VirtualProduction/LevelSnapshots/Source/LevelSnapshots/Private/Archive/ClassDefaults/BaseClassDefaultArchive.cpp
@@ -2,12 +2,11 @@
 
 #include "BaseClassDefaultArchive.h"
 
-#include "Util/Property/PropertyUtil.h"
-#include "Util/SnapshotObjectUtil.h"
+#include "Data/Util/Property/PropertyUtil.h"
+#include "Data/Util/WorldData/SnapshotObjectUtil.h"
 #include "WorldSnapshotData.h"
 
-#include "Components/ActorComponent.h"
-#include "GameFramework/Actor.h"
+#include "Util/Property/WorldReferenceCheckingUtil.h"
 
 bool UE::LevelSnapshots::Private::FBaseClassDefaultArchive::ShouldSkipProperty(const FProperty* InProperty) const
 {
@@ -37,24 +36,9 @@ UE::LevelSnapshots::Private::FBaseClassDefaultArchive::FBaseClassDefaultArchive(
 
 bool UE::LevelSnapshots::Private::FBaseClassDefaultArchive::IsPropertyReferenceToSubobjectOrClassDefaults(const FProperty* InProperty) const
 {
-	const FObjectPropertyBase* ObjectProperty = CastField<FObjectPropertyBase>(InProperty);
-	if (!ObjectProperty)
+	return UE::LevelSnapshots::Private::ContainsSubobjectOrSatisfiesPredicate(GetSerializedObject(), GetSerializedPropertyChain(), InProperty, [this](UObject* ContainedPtr)
 	{
-		return false;
-	}
-
-	const bool bIsMarkedAsSubobject = InProperty->HasAnyPropertyFlags(CPF_InstancedReference | CPF_ContainsInstancedReference | CPF_PersistentInstance);
-	const bool bIsActorOrComponentPtr = ObjectProperty->PropertyClass->IsChildOf(AActor::StaticClass()) || ObjectProperty->PropertyClass->IsChildOf(UActorComponent::StaticClass());
-	if (bIsMarkedAsSubobject || bIsActorOrComponentPtr)
-	{
-		return true;
-	}
-
-	const FArchiveSerializedPropertyChain* PropertyChain = GetSerializedPropertyChain();
-	void* ContainerPtr = GetSerializedObject();
-	const bool bIsUnsupported = UE::LevelSnapshots::Private::FollowPropertyChainUntilPredicateIsTrue(ContainerPtr, PropertyChain, InProperty, [this, ObjectProperty](void* LeafValuePtr)
-	{
-		if (const UObject* ContainedPtr = ObjectProperty->GetObjectPropertyValue(LeafValuePtr))
+		if (ContainedPtr)
 		{
 			const bool bIsClassDefault = ContainedPtr->HasAnyFlags(RF_ClassDefaultObject | RF_ArchetypeObject);
 			const bool bIsPointingToDefaultSubobject = ContainedPtr->HasAnyFlags(RF_DefaultSubObject);
@@ -62,9 +46,6 @@ bool UE::LevelSnapshots::Private::FBaseClassDefaultArchive::IsPropertyReferenceT
 			const bool bIsPointingToSubobject = ContainedPtr->IsIn(GetSerializedObject());
 			return bIsClassDefault || bIsPointingToDefaultSubobject || bIsPointingToSelf || bIsPointingToSubobject;
 		}
-	
 		return false;
-	});
-
-	return bIsUnsupported;
+	}); 
 }
\ No newline at end of file
diff --git a/Engine/Plugins/VirtualProduction/LevelSnapshots/Source/LevelSnapshots/Private/Archive/ClassDefaults/TakeClassDefaultObjectSnapshotArchive.cpp b/Engine/Plugins/VirtualProduction/LevelSnapshots/Source/LevelSnapshots/Private/Archive/ClassDefaults/TakeClassDefaultObjectSnapshotArchive.cpp
--- a/Engine/Plugins/VirtualProduction/LevelSnapshots/Source/LevelSnapshots/Private/Archive/ClassDefaults/TakeClassDefaultObjectSnapshotArchive.cpp
+++ b/Engine/Plugins/VirtualProduction/LevelSnapshots/Source/LevelSnapshots/Private/Archive/ClassDefaults/TakeClassDefaultObjectSnapshotArchive.cpp
@@ -5,10 +5,13 @@
 #include "ObjectSnapshotData.h"
 #include "WorldSnapshotData.h"
 
-void UE::LevelSnapshots::Private::FTakeClassDefaultObjectSnapshotArchive::SaveClassDefaultObject(FObjectSnapshotData& InObjectData, FWorldSnapshotData& InSharedData, UObject* SerializedObject)
+void UE::LevelSnapshots::Private::FTakeClassDefaultObjectSnapshotArchive::SaveClassDefaultObject(FClassSnapshotData& InObjectData, FWorldSnapshotData& InSharedData, UObject* SerializedObject)
 {
 	FTakeClassDefaultObjectSnapshotArchive SaveClass(InObjectData, InSharedData, SerializedObject);
 	SerializedObject->Serialize(SaveClass);
+	InObjectData.ObjectFlags = SerializedObject->GetFlags();
+	InObjectData.ClassPath = SerializedObject->GetClass();
+	InObjectData.ClassFlags = SerializedObject->GetClass()->GetFlags();
 }
 
 UE::LevelSnapshots::Private::FTakeClassDefaultObjectSnapshotArchive::FTakeClassDefaultObjectSnapshotArchive(FObjectSnapshotData& InObjectData, FWorldSnapshotData& InSharedData, UObject* InSerializedObject)
diff --git a/Engine/Plugins/VirtualProduction/LevelSnapshots/Source/LevelSnapshots/Private/Archive/ClassDefaults/TakeClassDefaultObjectSnapshotArchive.h b/Engine/Plugins/VirtualProduction/LevelSnapshots/Source/LevelSnapshots/Private/Archive/ClassDefaults/TakeClassDefaultObjectSnapshotArchive.h
--- a/Engine/Plugins/VirtualProduction/LevelSnapshots/Source/LevelSnapshots/Private/Archive/ClassDefaults/TakeClassDefaultObjectSnapshotArchive.h
+++ b/Engine/Plugins/VirtualProduction/LevelSnapshots/Source/LevelSnapshots/Private/Archive/ClassDefaults/TakeClassDefaultObjectSnapshotArchive.h
@@ -6,6 +6,7 @@
 #include "Archive/ClassDefaults/BaseClassDefaultArchive.h"
 
 class UObject;
+struct FClassSnapshotData;
 struct FObjectSnapshotData;
 struct FWorldSnapshotData;
 
@@ -17,7 +18,7 @@ namespace UE::LevelSnapshots::Private
 		using Super = FBaseClassDefaultArchive;
 	public:
 
-		static void SaveClassDefaultObject(FObjectSnapshotData& InObjectData, FWorldSnapshotData& InSharedData, UObject* SerializedObject);
+		static void SaveClassDefaultObject(FClassSnapshotData& InObjectData, FWorldSnapshotData& InSharedData, UObject* SerializedObject);
 
 	private:
 	
diff --git a/Engine/Plugins/VirtualProduction/LevelSnapshots/Source/LevelSnapshots/Private/Archive/LoadSnapshotObjectArchive.cpp b/Engine/Plugins/VirtualProduction/LevelSnapshots/Source/LevelSnapshots/Private/Archive/LoadSnapshotObjectArchive.cpp
--- a/Engine/Plugins/VirtualProduction/LevelSnapshots/Source/LevelSnapshots/Private/Archive/LoadSnapshotObjectArchive.cpp
+++ b/Engine/Plugins/VirtualProduction/LevelSnapshots/Source/LevelSnapshots/Private/Archive/LoadSnapshotObjectArchive.cpp
@@ -5,7 +5,7 @@
 #include "LevelSnapshotsLog.h"
 #include "LevelSnapshotsModule.h"
 #include "WorldSnapshotData.h"
-#include "Data/Util/SnapshotObjectUtil.h"
+#include "Data/Util/WorldData/SnapshotObjectUtil.h"
 
 #include "Serialization/ObjectWriter.h"
 #include "Internationalization/TextNamespaceUtil.h"
diff --git a/Engine/Plugins/VirtualProduction/LevelSnapshots/Source/LevelSnapshots/Private/Archive/TakeWorldObjectSnapshotArchive.cpp b/Engine/Plugins/VirtualProduction/LevelSnapshots/Source/LevelSnapshots/Private/Archive/TakeWorldObjectSnapshotArchive.cpp
--- a/Engine/Plugins/VirtualProduction/LevelSnapshots/Source/LevelSnapshots/Private/Archive/TakeWorldObjectSnapshotArchive.cpp
+++ b/Engine/Plugins/VirtualProduction/LevelSnapshots/Source/LevelSnapshots/Private/Archive/TakeWorldObjectSnapshotArchive.cpp
@@ -18,8 +18,8 @@ void UE::LevelSnapshots::Private::FTakeWorldObjectSnapshotArchive::TakeSnapshot(
 }
 
 UE::LevelSnapshots::Private::FTakeWorldObjectSnapshotArchive::FTakeWorldObjectSnapshotArchive(FObjectSnapshotData& InObjectData, FWorldSnapshotData& InSharedData, UObject* InOriginalObject)
-	:
-	Super(InObjectData, InSharedData, false, InOriginalObject)
+	: Super(InObjectData, InSharedData, false, InOriginalObject)
+	, Archetype(InOriginalObject->GetArchetype())
 {}
 
 bool UE::LevelSnapshots::Private::FTakeWorldObjectSnapshotArchive::ShouldSkipProperty(const FProperty* InProperty) const
@@ -48,7 +48,7 @@ bool UE::LevelSnapshots::Private::FTakeWorldObjectSnapshotArchive::ShouldSkipPro
 		}
 		
 		UObject* OriginalContainer = GetSerializedObject();
-		UObject* ClassDefaultContainer = OriginalContainer->GetClass()->GetDefaultObject(); 
+		UObject* ClassDefaultContainer = Archetype;
 		for (int32 ArrayDim = 0; ArrayDim < InProperty->ArrayDim; ++ArrayDim)
 		{
 			if (!InProperty->Identical_InContainer(OriginalContainer, ClassDefaultContainer, ArrayDim))
diff --git a/Engine/Plugins/VirtualProduction/LevelSnapshots/Source/LevelSnapshots/Private/Archive/TakeWorldObjectSnapshotArchive.h b/Engine/Plugins/VirtualProduction/LevelSnapshots/Source/LevelSnapshots/Private/Archive/TakeWorldObjectSnapshotArchive.h
--- a/Engine/Plugins/VirtualProduction/LevelSnapshots/Source/LevelSnapshots/Private/Archive/TakeWorldObjectSnapshotArchive.h
+++ b/Engine/Plugins/VirtualProduction/LevelSnapshots/Source/LevelSnapshots/Private/Archive/TakeWorldObjectSnapshotArchive.h
@@ -32,5 +32,7 @@ namespace UE::LevelSnapshots::Private
 	private:
 	
 		FTakeWorldObjectSnapshotArchive(FObjectSnapshotData& InObjectData, FWorldSnapshotData& InSharedData, UObject* InOriginalObject);
+
+		UObject* Archetype;
 	};
 }
\ No newline at end of file
diff --git a/Engine/Plugins/VirtualProduction/LevelSnapshots/Source/LevelSnapshots/Private/Data/CustomSerialization/CustomObjectSerializationWrapper.cpp b/Engine/Plugins/VirtualProduction/LevelSnapshots/Source/LevelSnapshots/Private/Data/CustomSerialization/CustomObjectSerializationWrapper.cpp
--- a/Engine/Plugins/VirtualProduction/LevelSnapshots/Source/LevelSnapshots/Private/Data/CustomSerialization/CustomObjectSerializationWrapper.cpp
+++ b/Engine/Plugins/VirtualProduction/LevelSnapshots/Source/LevelSnapshots/Private/Data/CustomSerialization/CustomObjectSerializationWrapper.cpp
@@ -5,9 +5,9 @@
 #include "Archive/ApplySnapshotToEditorArchive.h"
 #include "Archive/LoadSnapshotObjectArchive.h"
 #include "CustomSerialization/CustomSerializationDataManager.h"
-#include "Data/Util/SnapshotObjectUtil.h"
 #include "Data/Util/SnapshotUtil.h"
-#include "Data/Util/Restoration/ActorUtil.h"
+#include "Data/Util/WorldData/SnapshotObjectUtil.h"
+#include "Data/Util/WorldData/ActorUtil.h"
 #include "Data/WorldSnapshotData.h"
 #include "Interfaces/ICustomObjectSnapshotSerializer.h"
 #include "LevelSnapshotsLog.h"
@@ -129,7 +129,7 @@ namespace UE::LevelSnapshots::Private::Internal
 						[&WorldData, OriginalPath = MetaData->GetOriginalPath()](){ return UE::LevelSnapshots::Private::FindCustomSubobjectData(WorldData, OriginalPath);} );
 
 					FCustomSerializationData* SerializationData = SerializationDataGetter();
-					FApplySnapshotToEditorArchive::ApplyToRecreatedEditorWorldObject(SerializationData->Subobjects[i], WorldData, Cache, EditorSubobject, SelectionMap);
+					FApplySnapshotToEditorArchive::ApplyToEditorWorldObjectRecreatedWithArchetype(SerializationData->Subobjects[i], WorldData, Cache, EditorSubobject, SelectionMap);
 					CustomSerializer->OnPostSerializeEditorSubobject(EditorSubobject, *MetaData, SerializationDataReader);
 					continue;
 				}
@@ -142,7 +142,7 @@ namespace UE::LevelSnapshots::Private::Internal
 	}
 }
 
-void UE::LevelSnapshots::Private::TakeSnapshotForActor(
+void UE::LevelSnapshots::Private::TakeSnapshotOfActorCustomSubobjects(
 	AActor* EditorActor,
 	FCustomSerializationData& ActorSerializationData,
 	FWorldSnapshotData& WorldData)
diff --git a/Engine/Plugins/VirtualProduction/LevelSnapshots/Source/LevelSnapshots/Private/Data/CustomSerialization/CustomObjectSerializationWrapper.h b/Engine/Plugins/VirtualProduction/LevelSnapshots/Source/LevelSnapshots/Private/Data/CustomSerialization/CustomObjectSerializationWrapper.h
--- a/Engine/Plugins/VirtualProduction/LevelSnapshots/Source/LevelSnapshots/Private/Data/CustomSerialization/CustomObjectSerializationWrapper.h
+++ b/Engine/Plugins/VirtualProduction/LevelSnapshots/Source/LevelSnapshots/Private/Data/CustomSerialization/CustomObjectSerializationWrapper.h
@@ -49,7 +49,7 @@ namespace UE::LevelSnapshots::Private
 		}
 	};
 
-	void TakeSnapshotForActor(
+	void TakeSnapshotOfActorCustomSubobjects(
 			AActor* EditorActor,
 			FCustomSerializationData& ActorSerializationData,
 			FWorldSnapshotData& WorldData
diff --git a/Engine/Plugins/VirtualProduction/LevelSnapshots/Source/LevelSnapshots/Private/Data/CustomSerialization/CustomSerializationDataManager.cpp b/Engine/Plugins/VirtualProduction/LevelSnapshots/Source/LevelSnapshots/Private/Data/CustomSerialization/CustomSerializationDataManager.cpp
--- a/Engine/Plugins/VirtualProduction/LevelSnapshots/Source/LevelSnapshots/Private/Data/CustomSerialization/CustomSerializationDataManager.cpp
+++ b/Engine/Plugins/VirtualProduction/LevelSnapshots/Source/LevelSnapshots/Private/Data/CustomSerialization/CustomSerializationDataManager.cpp
@@ -9,7 +9,7 @@
 #include "Serialization/MemoryReader.h"
 #include "Serialization/NameAsStringProxyArchive.h"
 #include "Serialization/ObjectAndNameAsStringProxyArchive.h"
-#include "Util/SnapshotObjectUtil.h"
+#include "Util/WorldData/SnapshotObjectUtil.h"
 
 namespace UE::LevelSnapshots::Private::Internal
 {
diff --git a/Engine/Plugins/VirtualProduction/LevelSnapshots/Source/LevelSnapshots/Private/Data/LevelSnapshot.cpp b/Engine/Plugins/VirtualProduction/LevelSnapshots/Source/LevelSnapshots/Private/Data/LevelSnapshot.cpp
--- a/Engine/Plugins/VirtualProduction/LevelSnapshots/Source/LevelSnapshots/Private/Data/LevelSnapshot.cpp
+++ b/Engine/Plugins/VirtualProduction/LevelSnapshots/Source/LevelSnapshots/Private/Data/LevelSnapshot.cpp
@@ -7,14 +7,14 @@
 #include "Data/Util/ActorHashUtil.h"
 #include "Data/Util/EquivalenceUtil.h"
 #include "Data/Util/SnapshotUtil.h"
-#include "Data/Util/Restoration/ActorUtil.h"
+#include "Data/Util/WorldData/ActorUtil.h"
+#include "Data/Util/WorldData/WorldDataUtil.h"
 #include "LevelSnapshotsSettings.h"
 #include "LevelSnapshotsLog.h"
 #include "LevelSnapshotsModule.h"
 #include "Restorability/SnapshotRestorability.h"
 #include "SnapshotConsoleVariables.h"
 #include "Util/SortedScopedLog.h"
-#include "Util/Restoration/WorldDataUtil.h"
 
 #include "Algo/Accumulate.h"
 #include "Engine/Engine.h"
@@ -24,6 +24,7 @@
 #include "Misc/ScopeExit.h"
 #include "Stats/StatsMisc.h"
 #include "UObject/Package.h"
+#include "Util/WorldData/ClassDataUtil.h"
 #if WITH_EDITOR && !(UE_BUILD_SHIPPING || UE_BUILD_TEST)
 #include "Logging/MessageLog.h"
 #endif
@@ -83,6 +84,9 @@ bool ULevelSnapshot::SnapshotWorld(UWorld* TargetWorld)
 	// Existing snapshot world actors are now invalid - clear previous cached data
 	RecreateSnapshotWorld();
 	SerializedData = UE::LevelSnapshots::Private::SnapshotWorld(TargetWorld);
+	ClearCache();
+	Cache.InitFor(SerializedData);
+	// Should we delete preexisting actors from the snapshot world?
 
 	Module.OnPostTakeSnapshot().Broadcast({ this });
 
@@ -101,20 +105,28 @@ bool ULevelSnapshot::HasChangedSinceSnapshotWasTaken(AActor* WorldActor)
 		return false;
 	}
 
-#if WITH_EDITOR
+	// Address issue where unsupported actor were saved. This should be checked regardless though in case project settings were changes since the snapshot was taken.
+	if (!UE::LevelSnapshots::Restorability::IsActorDesirableForCapture(WorldActor))
+	{
+		return false;
+	}
+
 	ECachedDiffResult& CachedDiffResult = CachedDiffedActors.FindOrAdd(WorldActor);
-	if (CachedDiffResult != ECachedDiffResult::NotInitialised)
+	if (CachedDiffResult != ECachedDiffResult::NotInitialized
+#if !WITH_EDITOR
+		&& bIsDiffCacheEnabled
+#endif
+	)
 	{
 		return CachedDiffResult == ECachedDiffResult::HadChanges;
 	}
-#endif
 
 	const bool bHasChanged = [&]()
 	{
 		// Do not slow down old snapshots by computing hash if they had none saved
 		const bool bHasHashInfo = SerializedData.SnapshotVersionInfo.GetSnapshotCustomVersion() >= UE::LevelSnapshots::Private::FSnapshotCustomVersion::ActorHash;
 		// If the object is already loaded, there is no point in wasting time and computing the hash
-		const bool bNeedsHash = Cache.ActorCache.Find(WorldActor) != nullptr;
+		const bool bNeedsHash = Cache.ActorCache.Find(WorldActor) == nullptr;
 		
 		if (!bHasHashInfo || !bNeedsHash || !UE::LevelSnapshots::Private::HasMatchingHash(SavedActorData->Hash, WorldActor))
 		{
@@ -125,9 +137,7 @@ bool ULevelSnapshot::HasChangedSinceSnapshotWasTaken(AActor* WorldActor)
 		return false;
 	}();
 
-#if WITH_EDITOR
 	CachedDiffResult = bHasChanged ? ECachedDiffResult::HadChanges : ECachedDiffResult::HadNoChanges;
-#endif
 	return bHasChanged;
 }
 
@@ -162,15 +172,7 @@ int32 ULevelSnapshot::GetNumSavedActors() const
 
 namespace UE::LevelSnapshots::Private::Internal
 {
-	FSoftObjectPath ExtractPathWithoutSubobjects(UObject* Object)
-	{
-		int32 ColonIndex;
-		const FString Path = Object->GetPathName();
-		Path.FindChar(':', ColonIndex);
-		return Path.Left(ColonIndex);
-	}
-
-	void ConditionBreakOnActor(const FString& NameToSearchFor, const FSoftObjectPath& ActorPath)
+	void ConditionalBreakOnActor(const FString& NameToSearchFor, const FSoftObjectPath& ActorPath)
 	{
 		if (!NameToSearchFor.IsEmpty() && ActorPath.ToString().Contains(NameToSearchFor))
 		{
@@ -180,6 +182,11 @@ namespace UE::LevelSnapshots::Private::Internal
 }
 
 void ULevelSnapshot::DiffWorld(UWorld* World, FActorPathConsumer HandleMatchedActor, FActorPathConsumer HandleRemovedActor, FActorConsumer HandleAddedActor) const
+{
+	DiffWorld(World, FActorConsumer::CreateLambda([&HandleMatchedActor](AActor* Actor){ HandleMatchedActor.Execute(Actor); }), HandleRemovedActor, HandleAddedActor);
+}
+
+void ULevelSnapshot::DiffWorld(UWorld* World, FActorConsumer HandleMatchedActor, FActorPathConsumer HandleRemovedActor, FActorConsumer HandleAddedActor) const
 {
 	SCOPED_SNAPSHOT_CORE_TRACE(DiffWorld);
 	
@@ -193,6 +200,7 @@ void ULevelSnapshot::DiffWorld(UWorld* World, FActorPathConsumer HandleMatchedAc
 		UE_LOG(LogLevelSnapshots, Log, TEXT("Finished diffing snapshot"));
 	};
 
+	const FString DebugActorName = UE::LevelSnapshots::ConsoleVariables::CVarBreakOnDiffMatchedActor.GetValueOnAnyThread();
 
 	// Find actors that are not present in the snapshot
 	TSet<AActor*> AllActors;
@@ -200,18 +208,25 @@ void ULevelSnapshot::DiffWorld(UWorld* World, FActorPathConsumer HandleMatchedAc
 	{
 		SCOPED_SNAPSHOT_CORE_TRACE(DiffWorld_FindAllActors);
 
+		
 		const int32 NumActorsInWorld = Algo::Accumulate(World->GetLevels(), 0, [](int64 Size, const ULevel* Level){ return Size + Level->Actors.Num(); });
 		AllActors.Reserve(NumActorsInWorld);
 		for (ULevel* Level : World->GetLevels())
 		{
-			LoadedLevels.Add(UE::LevelSnapshots::Private::Internal::ExtractPathWithoutSubobjects(Level));
+			LoadedLevels.Add(UE::LevelSnapshots::Private::ExtractPathWithoutSubobjects(Level));
 			
 			for (AActor* ActorInLevel : Level->Actors)
 			{
 				AllActors.Add(ActorInLevel);
-			
+
 				// Warning: ActorInLevel can be null, e.g. when an actor was just removed from the world (and still in undo buffer)
-				if (IsValid(ActorInLevel) && !SerializedData.HasMatchingSavedActor(ActorInLevel) && UE::LevelSnapshots::Restorability::ShouldConsiderNewActorForRemoval(ActorInLevel))
+				if (!IsValid(ActorInLevel))
+				{
+					continue;
+				}
+			
+				UE::LevelSnapshots::Private::Internal::ConditionalBreakOnActor(DebugActorName, ActorInLevel);
+				if (!SerializedData.HasMatchingSavedActor(ActorInLevel) && UE::LevelSnapshots::Restorability::ShouldConsiderNewActorForRemoval(ActorInLevel))
 				{
 					HandleAddedActor.Execute(ActorInLevel);
 				}
@@ -227,10 +242,9 @@ void ULevelSnapshot::DiffWorld(UWorld* World, FActorPathConsumer HandleMatchedAc
 		ULevelSnapshotsSettings* Settings = GetMutableDefault<ULevelSnapshotsSettings>();
 		
 		const bool bShouldLogDiffWorldTimes = UE::LevelSnapshots::ConsoleVariables::CVarLogTimeDiffingMatchedActors.GetValueOnAnyThread();
-		const FString DebugActorName = UE::LevelSnapshots::ConsoleVariables::CVarBreakOnDiffMatchedActor.GetValueOnAnyThread();
 		FConditionalSortedScopedLog SortedItems(bShouldLogDiffWorldTimes);
 		
-		SerializedData.ForEachOriginalActor([&HandleMatchedActor, &HandleRemovedActor, &HandleAddedActor, &AllActors, &LoadedLevels, Settings, &DebugActorName, &SortedItems](const FSoftObjectPath& OriginalActorPath, const FActorSnapshotData& SavedData)
+		SerializedData.ForEachOriginalActor([this, &HandleMatchedActor, &HandleRemovedActor, &HandleAddedActor, &AllActors, &LoadedLevels, Settings, &DebugActorName, &SortedItems](const FSoftObjectPath& OriginalActorPath, const FActorSnapshotData& SavedData)
 		{
 			const FSoftObjectPath LevelPath = OriginalActorPath.GetAssetPathString();
 			if (!LoadedLevels.Contains(LevelPath))
@@ -249,10 +263,11 @@ void ULevelSnapshot::DiffWorld(UWorld* World, FActorPathConsumer HandleMatchedAc
 				return;
 			}
 
-			UClass* ActorClass = SavedData.ActorClass.TryLoadClass<AActor>();
+			const FSoftClassPath ClassPath = UE::LevelSnapshots::Private::GetClass(SavedData, SerializedData);
+			const UClass* ActorClass = ClassPath.TryLoadClass<AActor>();
 			if (!ActorClass)
 			{
-				UE_LOG(LogLevel, Warning, TEXT("Cannot find class %s. Saved actor %s will not be restored."), *SavedData.ActorClass.ToString(), *OriginalActorPath.ToString());
+				UE_LOG(LogLevel, Warning, TEXT("Cannot find class %s. Saved actor %s will not be restored."), *ClassPath.ToString(), *OriginalActorPath.ToString());
 				return;
 			}
 			if (Settings->SkippedClasses.SkippedClasses.Contains(ActorClass))
@@ -269,11 +284,11 @@ void ULevelSnapshot::DiffWorld(UWorld* World, FActorPathConsumer HandleMatchedAc
 			else
 			{
 				const FScopedLogItem Log = SortedItems.AddScopedLogItem(OriginalActorPath.ToString());
-				UE::LevelSnapshots::Private::Internal::ConditionBreakOnActor(DebugActorName, OriginalActorPath);
+				UE::LevelSnapshots::Private::Internal::ConditionalBreakOnActor(DebugActorName, OriginalActorPath);
 				SCOPED_SNAPSHOT_CORE_TRACE(HandleMatchedActor)
 				
 				UE_LOG(LogLevelSnapshots, VeryVerbose, TEXT("Matched actor %s"), *OriginalActorPath.ToString());
-				HandleMatchedActor.Execute(OriginalActorPath);
+				HandleMatchedActor.Execute(Cast<AActor>(ResolvedActor));
 			}
 		});
 	}
@@ -291,7 +306,7 @@ void ULevelSnapshot::SetSnapshotDescription(const FString& InSnapshotDescription
 
 void ULevelSnapshot::BeginDestroy()
 {
-	if (SnapshotContainerWorld)
+	if (RootSnapshotWorld)
 	{
 		DestroyWorld();
 	}
@@ -309,35 +324,37 @@ FString ULevelSnapshot::GenerateDebugLogInfo() const
 
 void ULevelSnapshot::RecreateSnapshotWorld()
 {
-	if (SnapshotContainerWorld)
+	if (RootSnapshotWorld)
 	{
 		DestroyWorld();
 	}
 	EnsureWorldInitialised();
 }
 
-void ULevelSnapshot::EnsureWorldInitialised()
+namespace UE::LevelSnapshots::Private
 {
-	if (SnapshotContainerWorld == nullptr)
-	{
-		SnapshotContainerWorld = NewObject<UWorld>(GetTransientPackage(), NAME_None);
-		SnapshotContainerWorld->WorldType = EWorldType::EditorPreview;
+	static UWorld* CreateWorld(UObject* Outer, FName WorldName);
+	static TSet<FName> ExtractLevelNames(const FWorldSnapshotData& WorldData);
+}
 
-		// Note: Do NOT create a FWorldContext for this world.
-		// If you do, the render thread will send render commands every tick (and crash cuz we do not init the scene below).
-		SnapshotContainerWorld->InitializeNewWorld(UWorld::InitializationValues()
-			.InitializeScenes(false)		// This is memory only world: no rendering
-            .AllowAudioPlayback(false)
-            .RequiresHitProxies(false)		
-            .CreatePhysicsScene(false)
-            .CreateNavigation(false)
-            .CreateAISystem(false)
-            .ShouldSimulatePhysics(false)
-			.EnableTraceCollision(false)
-            .SetTransactional(false)
-            .CreateFXSystem(false)
-            );
 
+void ULevelSnapshot::EnsureWorldInitialised()
+{
+	using namespace UE::LevelSnapshots::Private;
+	
+	if (RootSnapshotWorld == nullptr)
+	{
+		const FName RootWorldName = *ExtractPathWithoutSubobjects(MapPath).GetAssetName();
+		RootSnapshotWorld = CreateWorld(this, RootWorldName);
+		
+		TSet<FName> Sublevels = ExtractLevelNames(SerializedData);
+		Sublevels.Remove(RootWorldName);
+		SnapshotSublevels.Reserve(Sublevels.Num());
+		for (const FName& SublevelName : Sublevels)
+		{
+			SnapshotSublevels.Add(CreateWorld(this, SublevelName));
+		}
+		
 		// Destroy our temporary world when the editor (or game) world is destroyed. Reasons:
 		// 1. After unloading a map checks for world GC leaks; it would fatally crash if we did not clear here.
 		// 2. Our temp map stores a "copy" of actors from the original world: the original world is no longer relevant, so neither is our temp world.
@@ -352,17 +369,57 @@ void ULevelSnapshot::EnsureWorldInitialised()
 	            }
 	        });
 		}
+		
+		Cache.InitFor(SerializedData);
 #if WITH_EDITOR
-		FCoreUObjectDelegates::OnObjectModified.AddUObject(this, &ULevelSnapshot::ClearCachedDiffFlag);
+		FCoreUObjectDelegates::OnObjectModified.AddUObject(this, &ULevelSnapshot::ResetDiffCacheToUninitialized);
 #endif
 	}
 
-	SerializedData.SnapshotWorld = SnapshotContainerWorld;
+	SerializedData.SnapshotWorld = RootSnapshotWorld;
+	SerializedData.SnapshotSublevels = SnapshotSublevels;
+}
+
+namespace UE::LevelSnapshots::Private
+{
+	static UWorld* CreateWorld(UObject* Outer, FName WorldName)
+	{
+		UWorld* World = NewObject<UWorld>(Outer, WorldName);
+		World->WorldType = EWorldType::EditorPreview;
+
+		// Note: Do NOT create a FWorldContext for this world.
+		// If you do, the render thread will send render commands every tick (and crash cuz we do not init the scene below).
+		World->InitializeNewWorld(UWorld::InitializationValues()
+			.InitializeScenes(false)		// This is memory only world: no rendering
+			.AllowAudioPlayback(false)
+			.RequiresHitProxies(false)		
+			.CreatePhysicsScene(false)
+			.CreateNavigation(false)
+			.CreateAISystem(false)
+			.ShouldSimulatePhysics(false)
+			.EnableTraceCollision(false)
+			.SetTransactional(false)
+			.CreateFXSystem(false)
+			);
+
+		return World;
+	}
+
+	static TSet<FName> ExtractLevelNames(const FWorldSnapshotData& WorldData)
+	{
+		TSet<FName> Result;
+		for (auto SavedActorIt = WorldData.ActorData.CreateConstIterator(); SavedActorIt; ++SavedActorIt)
+		{
+			const FName WorldName = *ExtractPathWithoutSubobjects(SavedActorIt->Key).GetAssetName();
+			Result.Add(WorldName);
+		}
+		return Result;
+	}
 }
 
 void ULevelSnapshot::DestroyWorld()
 {
-	if (ensureAlwaysMsgf(SnapshotContainerWorld, TEXT("World was already destroyed.")))
+	if (ensureAlwaysMsgf(RootSnapshotWorld, TEXT("World was already destroyed.")))
 	{
 		if (ensure(GEngine))
 		{
@@ -371,29 +428,42 @@ void ULevelSnapshot::DestroyWorld()
 		}
 				
 		SerializedData.SnapshotWorld.Reset();
+		SerializedData.SnapshotSublevels.Reset();
 		ClearCache();
 	
-		SnapshotContainerWorld->CleanupWorld();
-		SnapshotContainerWorld = nullptr;
+		RootSnapshotWorld->CleanupWorld();
+		for (UWorld* Sublevel : SnapshotSublevels)
+		{
+			Sublevel->CleanupWorld();
+		}
+		
+		RootSnapshotWorld = nullptr;
+		SnapshotSublevels.Reset();
 	}
 }
 
-void ULevelSnapshot::ClearCache()
+void ULevelSnapshot::ClearCache() 
 {
 	Cache.Reset();
-	
-#if WITH_EDITOR
-	CachedDiffedActors.Reset();
-#endif
+	CachedDiffedActors.Reset(); 
 }
 
 #if WITH_EDITOR
-void ULevelSnapshot::ClearCachedDiffFlag(UObject* ModifiedObject)
+void ULevelSnapshot::ResetDiffCacheToUninitialized(TArrayView<AActor*> ModifiedActors)
+{
+	CachedDiffedActors.Reserve(ModifiedActors.Num());
+	for (AActor* ModifiedActor : ModifiedActors)
+	{
+		CachedDiffedActors.FindOrAdd(ModifiedActor) = ECachedDiffResult::NotInitialized;
+	}
+}
+
+void ULevelSnapshot::ResetDiffCacheToUninitialized(UObject* ModifiedObject)
 {
 	AActor* AsActor = ModifiedObject->IsA<AActor>() ? Cast<AActor>(ModifiedObject) : ModifiedObject->GetTypedOuter<AActor>(); 
 	if (AsActor && SerializedData.HasMatchingSavedActor(AsActor))
 	{
-		CachedDiffedActors.FindOrAdd(AsActor) = ECachedDiffResult::NotInitialised;
+		CachedDiffedActors.FindOrAdd(AsActor) = ECachedDiffResult::NotInitialized;
 	}
 }
 #endif
diff --git a/Engine/Plugins/VirtualProduction/LevelSnapshots/Source/LevelSnapshots/Private/Data/Util/Component/BaseComponentRestorer.h b/Engine/Plugins/VirtualProduction/LevelSnapshots/Source/LevelSnapshots/Private/Data/Util/Component/BaseComponentRestorer.h
--- a/Engine/Plugins/VirtualProduction/LevelSnapshots/Source/LevelSnapshots/Private/Data/Util/Component/BaseComponentRestorer.h
+++ b/Engine/Plugins/VirtualProduction/LevelSnapshots/Source/LevelSnapshots/Private/Data/Util/Component/BaseComponentRestorer.h
@@ -15,6 +15,7 @@
 #include "Components/ActorComponent.h"
 #include "GameFramework/Actor.h"
 #include "ComponentInstanceDataCache.h"
+#include "Util/WorldData/ClassDataUtil.h"
 
 class AActor;
 class UActorComponent;
@@ -103,7 +104,8 @@ namespace UE::LevelSnapshots::Private
 
 		UActorComponent* AllocateComponent(const FSoftObjectPath& ComponentPath, FSubobjectSnapshotData& SubobjectData, const FComponentSnapshotData& ComponentData, UObject* ComponentOuter) const
 		{
-			UClass* ComponentClass = SubobjectData.Class.TryLoadClass<UActorComponent>();
+			const FSoftClassPath ClassPath = GetClass(SubobjectData, WorldData);
+			UClass* ComponentClass = ClassPath.TryLoadClass<UActorComponent>();
 			const FName ComponentName = *UE::LevelSnapshots::Private::ExtractLastSubobjectName(ComponentPath);
 			const EObjectFlags ObjectFlags = SubobjectData.GetObjectFlags();
 			if constexpr (TDerived::IsRestoringIntoSnapshotWorld())
diff --git a/Engine/Plugins/VirtualProduction/LevelSnapshots/Source/LevelSnapshots/Private/Data/Util/Component/ExistingActorComponentRestorer.cpp b/Engine/Plugins/VirtualProduction/LevelSnapshots/Source/LevelSnapshots/Private/Data/Util/Component/ExistingActorComponentRestorer.cpp
--- a/Engine/Plugins/VirtualProduction/LevelSnapshots/Source/LevelSnapshots/Private/Data/Util/Component/ExistingActorComponentRestorer.cpp
+++ b/Engine/Plugins/VirtualProduction/LevelSnapshots/Source/LevelSnapshots/Private/Data/Util/Component/ExistingActorComponentRestorer.cpp
@@ -6,7 +6,7 @@
 #include "Archive/ApplySnapshotToEditorArchive.h"
 #include "Data/ActorSnapshotData.h"
 #include "Data/Util/EquivalenceUtil.h"
-#include "Data/Util/Restoration/ActorUtil.h"
+#include "Data/Util/WorldData/ActorUtil.h"
 #include "Data/WorldSnapshotData.h"
 #include "CustomSerialization/CustomObjectSerializationWrapper.h"
 #include "Data/Util/SnapshotUtil.h"
@@ -149,7 +149,7 @@ namespace UE::LevelSnapshots::Private::Internal
 		for (const FDeferredComponentSerializationTask& Task : DeferredSerializationsTasks)
 		{
 			const FRestoreObjectScope FinishRestore = PreSubobjectRestore_EditorWorld(Task.Snapshot, Task.Original, GetWorldData(), Cache, SelectionMap, LocalisationSnapshotPackage);
-			FApplySnapshotToEditorArchive::ApplyToRecreatedEditorWorldObject(*Task.SubobjectData, GetWorldData(), Cache, Task.Original, SelectionMap); 
+			FApplySnapshotToEditorArchive::ApplyToEditorWorldObjectRecreatedWithArchetype(*Task.SubobjectData, GetWorldData(), Cache, Task.Original, SelectionMap); 
 		}
 		for (UActorComponent* Fail : CompsWithMissingSnapshotCounterpart)
 		{
diff --git a/Engine/Plugins/VirtualProduction/LevelSnapshots/Source/LevelSnapshots/Private/Data/Util/EquivalenceUtil.cpp b/Engine/Plugins/VirtualProduction/LevelSnapshots/Source/LevelSnapshots/Private/Data/Util/EquivalenceUtil.cpp
--- a/Engine/Plugins/VirtualProduction/LevelSnapshots/Source/LevelSnapshots/Private/Data/Util/EquivalenceUtil.cpp
+++ b/Engine/Plugins/VirtualProduction/LevelSnapshots/Source/LevelSnapshots/Private/Data/Util/EquivalenceUtil.cpp
@@ -5,7 +5,7 @@
 #include "Archive/ApplySnapshotToEditorArchive.h"
 #include "Data/CustomSerialization/CustomObjectSerializationWrapper.h"
 #include "Data/SnapshotCustomVersion.h"
-#include "Data/Util/Restoration/ActorUtil.h"
+#include "Data/Util/WorldData/ActorUtil.h"
 #include "Data/WorldSnapshotData.h"
 #include "LevelSnapshotsLog.h"
 #include "LevelSnapshotsModule.h"
@@ -321,21 +321,6 @@ bool UE::LevelSnapshots::Private::AreObjectPropertiesEquivalent(ULevelSnapshot*
 
 namespace UE::LevelSnapshots::Private::Internal
 {
-	/** Checks whether the two actors object properties should be considered equivalent */
-	static bool AreActorsEquivalent(UObject* SnapshotPropertyValue, AActor* OriginalActorReference, const TMap<FSoftObjectPath, FActorSnapshotData>& ActorData, const FSnapshotDataCache& Cache)
-	{
-		// Compare actors
-		const FActorSnapshotData* SavedData = ActorData.Find(OriginalActorReference);
-		if (SavedData == nullptr)
-		{
-			return false;
-		}
-
-		// The snapshot actor was already allocated, if some other snapshot actor is referencing it
-		const TOptional<TNonNullPtr<AActor>> PreallocatedSnapshotVersion = UE::LevelSnapshots::Private::GetPreallocatedIfCached(OriginalActorReference, Cache);
-		return PreallocatedSnapshotVersion.Get(nullptr) == SnapshotPropertyValue;
-	}
-	
 	/** Checks whether the two subobject object properties should be considered equivalent */
 	static bool HaveSameNames(UObject* SnapshotPropertyValue, UObject* OriginalPropertyValue, const FSnapshotDataCache& Cache)
 	{
@@ -376,7 +361,7 @@ namespace UE::LevelSnapshots::Private::Internal
 		AActor* OwningWorldActor = OriginalPropertyValue->GetTypedOuter<AActor>();
 		if (ensure(OwningSnapshotActor && OwningWorldActor))
 		{
-			return AreActorsEquivalent(OwningSnapshotActor, OwningWorldActor, Snapshot->GetSerializedData().ActorData, Snapshot->GetCache()) && !HaveDifferentPropertyValues(Snapshot, SnapshotPropertyValue, OriginalPropertyValue, SnapshotActor, OriginalActor);
+			return AreActorsEquivalent(OwningSnapshotActor, OwningWorldActor, Snapshot->GetSerializedData(), Snapshot->GetCache()) && !HaveDifferentPropertyValues(Snapshot, SnapshotPropertyValue, OriginalPropertyValue, SnapshotActor, OriginalActor);
 		}
 		return false;
 	}
@@ -400,7 +385,7 @@ bool UE::LevelSnapshots::Private::AreReferencesEquivalent(ULevelSnapshot* Snapsh
 
 	if (AActor* OriginalActorReference = Cast<AActor>(OriginalPropertyValue))
 	{
-		return Internal::AreActorsEquivalent(SnapshotPropertyValue, OriginalActorReference, Snapshot->GetSerializedData().ActorData, Snapshot->GetCache());
+		return AreActorsEquivalent(SnapshotPropertyValue, OriginalActorReference, Snapshot->GetSerializedData(), Snapshot->GetCache());
 	}
 	
 	const bool bIsWorldObject = SnapshotPropertyValue->IsInA(UWorld::StaticClass()) && OriginalPropertyValue->IsInA(UWorld::StaticClass());
@@ -414,4 +399,18 @@ bool UE::LevelSnapshots::Private::AreReferencesEquivalent(ULevelSnapshot* Snapsh
 	}
 	
 	return SnapshotPropertyValue == OriginalPropertyValue;
+}
+
+bool UE::LevelSnapshots::Private::AreActorsEquivalent(UObject* SnapshotPropertyValue, AActor* OriginalActorReference, const FWorldSnapshotData& WorldData, const FSnapshotDataCache& Cache)
+{
+	// Compare actors
+	const FActorSnapshotData* SavedData = WorldData.ActorData.Find(OriginalActorReference);
+	if (SavedData == nullptr)
+	{
+		return false;
+	}
+
+	// The snapshot actor was already allocated, if some other snapshot actor is referencing it
+	const TOptional<TNonNullPtr<AActor>> PreallocatedSnapshotVersion = UE::LevelSnapshots::Private::GetPreallocatedIfCached(OriginalActorReference, Cache);
+	return PreallocatedSnapshotVersion.Get(nullptr) == SnapshotPropertyValue;
 }
\ No newline at end of file
diff --git a/Engine/Plugins/VirtualProduction/LevelSnapshots/Source/LevelSnapshots/Private/Data/Util/EquivalenceUtil.h b/Engine/Plugins/VirtualProduction/LevelSnapshots/Source/LevelSnapshots/Private/Data/Util/EquivalenceUtil.h
--- a/Engine/Plugins/VirtualProduction/LevelSnapshots/Source/LevelSnapshots/Private/Data/Util/EquivalenceUtil.h
+++ b/Engine/Plugins/VirtualProduction/LevelSnapshots/Source/LevelSnapshots/Private/Data/Util/EquivalenceUtil.h
@@ -55,4 +55,7 @@ namespace UE::LevelSnapshots::Private
 	
 	/** Checks whether two pointers point to "equivalent" objects. */
 	bool AreReferencesEquivalent(ULevelSnapshot* Snapshot, UObject* SnapshotPropertyValue, UObject* OriginalPropertyValue, AActor* SnapshotActor, AActor* OriginalActor);
+
+	/** Checks whether the two actors are equivalent */
+	bool AreActorsEquivalent(UObject* SnapshotPropertyValue, AActor* OriginalActorReference, const FWorldSnapshotData& WorldData, const FSnapshotDataCache& Cache);
 }
diff --git a/Engine/Plugins/VirtualProduction/LevelSnapshots/Source/LevelSnapshots/Private/Data/Util/Property/PropertyUtil.cpp b/Engine/Plugins/VirtualProduction/LevelSnapshots/Source/LevelSnapshots/Private/Data/Util/Property/PropertyUtil.cpp
--- a/Engine/Plugins/VirtualProduction/LevelSnapshots/Source/LevelSnapshots/Private/Data/Util/Property/PropertyUtil.cpp
+++ b/Engine/Plugins/VirtualProduction/LevelSnapshots/Source/LevelSnapshots/Private/Data/Util/Property/PropertyUtil.cpp
@@ -2,80 +2,128 @@
 
 #include "PropertyUtil.h"
 
+#include "Serialization/ArchiveSerializedPropertyChain.h"
 #include "UObject/Field.h"
 #include "UObject/UnrealType.h"
-#include "Serialization/ArchiveSerializedPropertyChain.h"
 
-namespace UE::LevelSnapshots::Private::Internal
+namespace UE::LevelSnapshots::Private
 {
-	static EBreakBehaviour RecursiveFollowPropertyChain(int32 ChainIndex, void* ContainerPtr, const FArchiveSerializedPropertyChain* PropertyChain, const FProperty* LeafProperty, const FHandleValuePtr& Callback)
+	namespace Internal
 	{
-		const bool bHasReachLeafProperty = !PropertyChain || PropertyChain->GetNumProperties() <= ChainIndex; 
-		if (bHasReachLeafProperty)
-		{
-			void* LeafValuePtr = LeafProperty->ContainerPtrToValuePtr<void>(ContainerPtr);
-			return Callback(LeafValuePtr);
-		}
-		
-		const FProperty* CurrentProperty = PropertyChain->GetPropertyFromRoot(ChainIndex);
-		void* ValuePtr = CurrentProperty->ContainerPtrToValuePtr<void>(ContainerPtr);
-		
-		if (const FArrayProperty* ArrayProperty = CastField<FArrayProperty>(CurrentProperty))
+		static EBreakBehaviour RecursiveFollowPropertyChain(int32 ChainIndex, void* ContainerPtr, const FArchiveSerializedPropertyChain* PropertyChain, const FProperty* LeafProperty, const FHandleValuePtr& Callback)
 		{
-			FScriptArrayHelper ArrayHelper(ArrayProperty, ValuePtr);
+			const bool bHasReachLeafProperty = !PropertyChain || PropertyChain->GetNumProperties() <= ChainIndex;
 			
-			bool bShouldStop = false;
-			for (int32 i = 0; i < ArrayHelper.Num() && !bShouldStop; ++i)
+			const FProperty* CurrentProperty = bHasReachLeafProperty
+				? LeafProperty
+				: PropertyChain->GetPropertyFromRoot(ChainIndex);
+			void* ValuePtr = CurrentProperty->ContainerPtrToValuePtr<void>(ContainerPtr);
+			if (const FArrayProperty* ArrayProperty = CastField<FArrayProperty>(CurrentProperty))
 			{
-				void* IndexValuePtr = ArrayHelper.GetRawPtr(i);
-				void* ElementValuePtr = ArrayProperty->Inner->ContainerPtrToValuePtr<void>(IndexValuePtr);
+				FScriptArrayHelper ArrayHelper(ArrayProperty, ValuePtr);
 				
-				const EBreakBehaviour CallbackResult = RecursiveFollowPropertyChain(ChainIndex + 1, ElementValuePtr, PropertyChain, LeafProperty, Callback);
-				bShouldStop |= CallbackResult == EBreakBehaviour::Break;
+				bool bShouldStop = false;
+				for (int32 i = 0; i < ArrayHelper.Num() && !bShouldStop; ++i)
+				{
+					void* IndexValuePtr = ArrayHelper.GetRawPtr(i);
+					void* ElementValuePtr = ArrayProperty->Inner->ContainerPtrToValuePtr<void>(IndexValuePtr);
+
+					const EBreakBehaviour CallbackResult = bHasReachLeafProperty
+						? Callback(ElementValuePtr, EPropertyType::NormalProperty)
+						: RecursiveFollowPropertyChain(ChainIndex + 1, ElementValuePtr, PropertyChain, LeafProperty, Callback);
+					bShouldStop |= CallbackResult == EBreakBehaviour::Break;
+				}
+
+				return bShouldStop ? EBreakBehaviour::Break : EBreakBehaviour::Continue;
 			}
 
-			return EBreakBehaviour::Continue;
-		}
-		else if (const FSetProperty* SetProperty = CastField<FSetProperty>(CurrentProperty))
-		{
-			FScriptSetHelper SetHelper(SetProperty, ValuePtr);
-			
-			bool bShouldStop = false;
-			for (int32 i = 0; i < SetHelper.Num() && !bShouldStop; ++i)
+			if (const FSetProperty* SetProperty = CastField<FSetProperty>(CurrentProperty))
 			{
-				void* IndexValuePtr = SetHelper.GetElementPtr(i);
-				void* ElementValuePtr = SetProperty->ElementProp->ContainerPtrToValuePtr<void>(IndexValuePtr);
+				FScriptSetHelper SetHelper(SetProperty, ValuePtr);
 				
-				const EBreakBehaviour CallbackResult = RecursiveFollowPropertyChain(ChainIndex + 1, ElementValuePtr, PropertyChain, LeafProperty, Callback);
-				bShouldStop |= CallbackResult == EBreakBehaviour::Break;
+				bool bShouldStop = false;
+				for (int32 i = 0; i < SetHelper.Num() && !bShouldStop; ++i)
+				{
+					void* IndexContainerPtr = SetHelper.GetElementPtr(i);
+					void* ElementValuePtr = SetProperty->ElementProp->ContainerPtrToValuePtr<void>(IndexContainerPtr);
+
+					const EBreakBehaviour CallbackResult = bHasReachLeafProperty
+						? Callback(ElementValuePtr, EPropertyType::NormalProperty)
+						: RecursiveFollowPropertyChain(ChainIndex + 1, ElementValuePtr, PropertyChain, LeafProperty, Callback);
+					bShouldStop |= CallbackResult == EBreakBehaviour::Break;
+				}
+
+				return bShouldStop ? EBreakBehaviour::Break : EBreakBehaviour::Continue;
 			}
 
-			return EBreakBehaviour::Continue;
-		}
-		else if (const FMapProperty* MapProperty = CastField<FMapProperty>(CurrentProperty))
-		{
-			// Currently unsupported
-			return EBreakBehaviour::Break;
-		}
-		else
-		{
+			if (const FMapProperty* MapProperty = CastField<FMapProperty>(CurrentProperty))
+			{
+				FScriptMapHelper MapHelper(MapProperty, ValuePtr);
+
+				bool bShouldStop = false;
+				for (int32 i = 0; i < MapHelper.Num() && !bShouldStop; ++i)
+				{
+					// If the map is the leaf, check content of both key and value
+					if (bHasReachLeafProperty)
+					{
+						void* IndexKeyContainerPtr = MapHelper.GetKeyPtr(i);
+						const EBreakBehaviour KeyCallbackResult = bHasReachLeafProperty
+							? Callback(IndexKeyContainerPtr, EPropertyType::KeyInMap)
+							: RecursiveFollowPropertyChain(ChainIndex + 1, IndexKeyContainerPtr, PropertyChain, LeafProperty, Callback);
+						bShouldStop |= KeyCallbackResult == EBreakBehaviour::Break;
+
+						if (!bShouldStop) // Fast path
+						{
+							void* IndexValueContainerPtr = MapHelper.GetValuePtr(i);
+							const EBreakBehaviour ValueCallbackResult = bHasReachLeafProperty
+								? Callback(IndexValueContainerPtr, EPropertyType::ValueInMap)
+								: RecursiveFollowPropertyChain(ChainIndex + 1, IndexValueContainerPtr, PropertyChain, LeafProperty, Callback);
+							bShouldStop |= ValueCallbackResult == EBreakBehaviour::Break;
+						}
+					}
+					else // If the map is in the middle of the chain, continue down the correct path
+					{
+						const FProperty* NextProperty = PropertyChain->GetPropertyFromRoot(ChainIndex + 1);
+						if (NextProperty == MapHelper.GetKeyProperty())
+						{
+							void* IndexKeyContainerPtr = MapHelper.GetPairPtr(i);
+							const EBreakBehaviour KeyCallbackResult = RecursiveFollowPropertyChain(ChainIndex + 1, IndexKeyContainerPtr, PropertyChain, LeafProperty, Callback);
+							bShouldStop |= KeyCallbackResult == EBreakBehaviour::Break;
+						}
+						else if (ensure(NextProperty == MapHelper.GetValueProperty()))
+						{
+							void* IndexValueContainerPtr = MapHelper.GetPairPtr(i);
+							const EBreakBehaviour ValueCallbackResult = RecursiveFollowPropertyChain(ChainIndex + 1, IndexValueContainerPtr, PropertyChain, LeafProperty, Callback);
+							bShouldStop |= ValueCallbackResult == EBreakBehaviour::Break;
+						}
+					}
+				}
+
+				return bShouldStop ? EBreakBehaviour::Break : EBreakBehaviour::Continue;
+			}
+			
+			if (bHasReachLeafProperty)
+			{
+				void* LeafValuePtr = LeafProperty->ContainerPtrToValuePtr<void>(ContainerPtr);
+				return Callback(LeafValuePtr, EPropertyType::NormalProperty); 
+			}
 			return RecursiveFollowPropertyChain(ChainIndex + 1, ValuePtr, PropertyChain, LeafProperty, Callback);
 		}
 	}
-}
-
-void UE::LevelSnapshots::Private::FollowPropertyChain(void* ContainerPtr, const FArchiveSerializedPropertyChain* PropertyChain, const FProperty* LeafProperty, UE::LevelSnapshots::Private::FHandleValuePtr Callback)
-{
-	Internal::RecursiveFollowPropertyChain(0, ContainerPtr, PropertyChain, LeafProperty, Callback);
-}
+	
+	void FollowPropertyChain(void* ContainerPtr, const FArchiveSerializedPropertyChain* PropertyChain, const FProperty* LeafProperty, UE::LevelSnapshots::Private::FHandleValuePtr Callback)
+	{
+		Internal::RecursiveFollowPropertyChain(0, ContainerPtr, PropertyChain, LeafProperty, Callback);
+	}
 
-bool UE::LevelSnapshots::Private::FollowPropertyChainUntilPredicateIsTrue(void* ContainerPtr, const FArchiveSerializedPropertyChain* PropertyChain, const FProperty* LeafProperty, UE::LevelSnapshots::Private::FValuePtrPredicate Callback)
-{
-	bool bResult = false;
-	FollowPropertyChain(ContainerPtr, PropertyChain, LeafProperty, [Callback = MoveTemp(Callback), &bResult](void* ValuePtr)
+	bool FollowPropertyChainUntilPredicateIsTrue(void* ContainerPtr, const FArchiveSerializedPropertyChain* PropertyChain, const FProperty* LeafProperty, UE::LevelSnapshots::Private::FValuePtrPredicate Callback)
 	{
-		bResult |= Callback(ValuePtr);
-		return bResult ? EBreakBehaviour::Break : EBreakBehaviour::Continue;
-	});
-	return bResult;
+		bool bResult = false;
+		FollowPropertyChain(ContainerPtr, PropertyChain, LeafProperty, [Callback = MoveTemp(Callback), &bResult](void* ValuePtr, EPropertyType PropertyType)
+		{
+			bResult |= Callback(ValuePtr, PropertyType);
+			return bResult ? EBreakBehaviour::Break : EBreakBehaviour::Continue;
+		});
+		return bResult;
+	}
 }
\ No newline at end of file
diff --git a/Engine/Plugins/VirtualProduction/LevelSnapshots/Source/LevelSnapshots/Private/Data/Util/Property/PropertyUtil.h b/Engine/Plugins/VirtualProduction/LevelSnapshots/Source/LevelSnapshots/Private/Data/Util/Property/PropertyUtil.h
--- a/Engine/Plugins/VirtualProduction/LevelSnapshots/Source/LevelSnapshots/Private/Data/Util/Property/PropertyUtil.h
+++ b/Engine/Plugins/VirtualProduction/LevelSnapshots/Source/LevelSnapshots/Private/Data/Util/Property/PropertyUtil.h
@@ -6,6 +6,7 @@
 #include "Templates/Function.h"
 #include "UObject/ObjectMacros.h"
 
+class FFieldClass;
 class FProperty;
 struct FArchiveSerializedPropertyChain;
 
@@ -16,8 +17,17 @@ namespace UE::LevelSnapshots::Private
 		Continue,
 		Break
 	};
-	using FHandleValuePtr = TFunctionRef<EBreakBehaviour(void* ValuePtr)>;
-	using FValuePtrPredicate = TFunctionRef<bool(void* ValuePtr)>;
+	enum class EPropertyType
+	{
+		/** Any other non-map property */
+		NormalProperty,
+		/** This is the key of a map property */
+		KeyInMap,
+		/** This is the value of a map property */
+		ValueInMap
+	};
+	using FHandleValuePtr = TFunctionRef<EBreakBehaviour(void* ValuePtr, EPropertyType PropertyType)>;
+	using FValuePtrPredicate = TFunctionRef<bool(void* ValuePtr, EPropertyType PropertyType)>;
 	
 	/**
 	 * Follow a property chain to its leaf property and call a function with the value pointers.
diff --git a/Engine/Plugins/VirtualProduction/LevelSnapshots/Source/LevelSnapshots/Private/Data/Util/Property/WorldReferenceCheckingUtil.cpp b/Engine/Plugins/VirtualProduction/LevelSnapshots/Source/LevelSnapshots/Private/Data/Util/Property/WorldReferenceCheckingUtil.cpp
new file mode 100644
--- /dev/null
+++ b/Engine/Plugins/VirtualProduction/LevelSnapshots/Source/LevelSnapshots/Private/Data/Util/Property/WorldReferenceCheckingUtil.cpp
@@ -0,0 +1,105 @@
+// Copyright Epic Games, Inc. All Rights Reserved.
+
+#include "WorldReferenceCheckingUtil.h"
+
+#include "Components/ActorComponent.h"
+#include "Engine/World.h"
+#include "GameFramework/Actor.h"
+#include "Serialization/ArchiveSerializedPropertyChain.h"
+#include "UObject/Field.h"
+#include "UObject/UnrealType.h"
+
+namespace UE::LevelSnapshots::Private
+{
+	namespace Internal
+	{
+		class FSubobjectOrPredicateChecker
+		{
+			UObject* OwningWorldObject;
+			const FArchiveSerializedPropertyChain* PropertyChain;
+			const FUObjectFilter UObjectFilterFunc;
+		public:
+
+			FSubobjectOrPredicateChecker(UObject* OwningWorldObject, const FArchiveSerializedPropertyChain* PropertyChain, FUObjectFilter UObjectFilterFunc)
+				: OwningWorldObject(OwningWorldObject)
+				, PropertyChain(PropertyChain)
+				, UObjectFilterFunc(UObjectFilterFunc)
+			{}
+
+			bool ContainsWorldObjectProperty(const FProperty* InLeafMemberProperty) const
+			{
+				if (const FObjectPropertyBase* ObjectProperty = CastField<FObjectPropertyBase>(InLeafMemberProperty))
+				{
+					return IsWorldObjectReference(ObjectProperty, ObjectProperty, EPropertyType::NormalProperty);
+				}
+				
+				if (const FArrayProperty* ArrayProperty = CastField<FArrayProperty>(InLeafMemberProperty))
+				{
+					if (const FObjectPropertyBase* InnerAsObjectProperty = CastField<FObjectPropertyBase>(ArrayProperty->Inner))
+					{
+						return IsWorldObjectReference(InnerAsObjectProperty, ArrayProperty, EPropertyType::NormalProperty);
+					}
+				}
+				
+				if (const FSetProperty* SetProperty = CastField<FSetProperty>(InLeafMemberProperty))
+				{
+					if (const FObjectPropertyBase* InnerAsObjectProperty = CastField<FObjectPropertyBase>(SetProperty->ElementProp))
+					{
+						return IsWorldObjectReference(InnerAsObjectProperty, SetProperty, EPropertyType::NormalProperty);
+					}
+				}
+
+				if (const FMapProperty* MapProperty = CastField<FMapProperty>(InLeafMemberProperty))
+				{
+					bool bResult = false;
+					// In the case where both key and value are object properties, the efficiency could be improved... this API less complex though
+					if (const FObjectPropertyBase* KeyAsObjectProperty = CastField<FObjectPropertyBase>(MapProperty->KeyProp))
+					{
+						bResult |= IsWorldObjectReference(KeyAsObjectProperty, MapProperty, EPropertyType::KeyInMap);
+					}
+					if (const FObjectPropertyBase* ValueAsObjectProperty = CastField<FObjectPropertyBase>(MapProperty->ValueProp);
+						ValueAsObjectProperty && bResult) // Fast path out
+					{
+						bResult |= IsWorldObjectReference(ValueAsObjectProperty, MapProperty, EPropertyType::ValueInMap);
+					}
+					return bResult;
+				}
+
+				return false;
+			}
+			
+			/**
+			 * Loops to the end of the property chain checking whether the ObjectProperty points to a world object.
+			 * @param ObjectProperty The property on which to perform the reference check.
+			 * @param LeafProperty The property that contains ObjectProperty. If ObjectProperty is not the inner property of an array, set or map then LeafProperty == ObjectProperty.
+			 */
+			bool IsWorldObjectReference(const FObjectPropertyBase* ObjectProperty, const FProperty* LeafProperty, EPropertyType PropertyType) const
+			{
+				const bool bIsMarkedAsSubobject = LeafProperty->HasAnyPropertyFlags(CPF_InstancedReference | CPF_ContainsInstancedReference | CPF_PersistentInstance);
+				const bool bIsActorOrComponentPtr = PropertyType == EPropertyType::NormalProperty
+					&& (ObjectProperty->PropertyClass->IsChildOf(AActor::StaticClass()) || ObjectProperty->PropertyClass->IsChildOf(UActorComponent::StaticClass()));
+				if (bIsMarkedAsSubobject || bIsActorOrComponentPtr)
+				{
+					return true;
+				}
+
+				const bool bSatisfiesPredicate = FollowPropertyChainUntilPredicateIsTrue(OwningWorldObject, PropertyChain, LeafProperty, [this, ObjectProperty, PropertyType](void* LeafValuePtr, EPropertyType ChainPropertyType)
+				{
+					if (PropertyType == ChainPropertyType)
+					{
+						UObject* ContainedPtr = ObjectProperty->GetObjectPropertyValue(LeafValuePtr);
+						return LeafValuePtr ? UObjectFilterFunc(ContainedPtr) : false;
+					}
+					return false;
+				});
+				return bSatisfiesPredicate;
+			}
+		};
+	}
+	
+	bool ContainsSubobjectOrSatisfiesPredicate(UObject* OwningWorldObject, const FArchiveSerializedPropertyChain* PropertyChain, const FProperty* LeafMemberProperty, FUObjectFilter Predicate)
+	{
+		Internal::FSubobjectOrPredicateChecker Helper(OwningWorldObject, PropertyChain, Predicate);
+		return Helper.ContainsWorldObjectProperty(LeafMemberProperty);
+	}
+}
\ No newline at end of file
diff --git a/Engine/Plugins/VirtualProduction/LevelSnapshots/Source/LevelSnapshots/Private/Data/Util/Property/WorldReferenceCheckingUtil.h b/Engine/Plugins/VirtualProduction/LevelSnapshots/Source/LevelSnapshots/Private/Data/Util/Property/WorldReferenceCheckingUtil.h
new file mode 100644
--- /dev/null
+++ b/Engine/Plugins/VirtualProduction/LevelSnapshots/Source/LevelSnapshots/Private/Data/Util/Property/WorldReferenceCheckingUtil.h
@@ -0,0 +1,34 @@
+// Copyright Epic Games, Inc. All Rights Reserved.
+
+#pragma once
+
+#include "CoreMinimal.h"
+#include "PropertyUtil.h"
+
+class FProperty;
+class UObject;
+
+struct FArchiveSerializedPropertyChain;
+
+namespace UE::LevelSnapshots::Private
+{
+	using FUObjectFilter = TFunction<bool(UObject*)>;
+	
+	/**
+	 * Helper for detecting whether a property on a serialized world object contains a subobject reference or satifies the predicate.
+	 * The predicate is usually something that checks whether the reference is a world object or a class default reference.
+	 * This handles single object properties as well as collections (arrays, sets, maps).
+	 * 
+	 * Examples: calling
+	 *  - ContainsWorldObjectReference(SceneComponent, {}, AttachChildren) would return true (if the USceneComponent::AttachChildren is non-empty).
+	 *  - ContainsWorldObjectReference(StaticMeshComponent, {}, UStaticMeshComponent::StaticMesh) would return false.
+	 *
+	 *  @param OwningWorldObject Some object that exists in a world
+	 *  @param PropertyChain The properties leading up to LeafProperty - nullptr implies LeafProperty is a root property
+	 *  @param LeafMemberProperty The last member property (i.e. not inner of any array, set, or map) at the end of the chain (the last element in the chain does not contain LeafProperty).
+	 *  @param Predicate Checks some condition on a UObject that was contained by the property; the passed in UObject can be nullptr.
+	 *
+	 *  @return Whether any value contained by any LeafProperty contained a world reference. 
+	 */
+	bool ContainsSubobjectOrSatisfiesPredicate(UObject* OwningWorldObject, const FArchiveSerializedPropertyChain* PropertyChain, const FProperty* LeafMemberProperty, FUObjectFilter Predicate);
+};
diff --git a/Engine/Plugins/VirtualProduction/LevelSnapshots/Source/LevelSnapshots/Private/Data/Util/SnapshotUtil.cpp b/Engine/Plugins/VirtualProduction/LevelSnapshots/Source/LevelSnapshots/Private/Data/Util/SnapshotUtil.cpp
--- a/Engine/Plugins/VirtualProduction/LevelSnapshots/Source/LevelSnapshots/Private/Data/Util/SnapshotUtil.cpp
+++ b/Engine/Plugins/VirtualProduction/LevelSnapshots/Source/LevelSnapshots/Private/Data/Util/SnapshotUtil.cpp
@@ -9,6 +9,16 @@
 #include "GameFramework/Actor.h"
 #include "Misc/ScopeExit.h"
 
+FSoftObjectPath UE::LevelSnapshots::Private::ExtractPathWithoutSubobjects(const FSoftObjectPath& ObjectPath)
+{
+	int32 ColonIndex = INDEX_NONE;
+	const FString Path = ObjectPath.ToString();
+	Path.FindChar(':', ColonIndex);
+	return ColonIndex == INDEX_NONE
+		? Path
+		: Path.Left(ColonIndex);
+}
+
 FString UE::LevelSnapshots::Private::ExtractLastSubobjectName(const FSoftObjectPath& ObjectPath)
 {
 	const FString& SubPathString = ObjectPath.GetSubPathString();
diff --git a/Engine/Plugins/VirtualProduction/LevelSnapshots/Source/LevelSnapshots/Private/Data/Util/SnapshotUtil.h b/Engine/Plugins/VirtualProduction/LevelSnapshots/Source/LevelSnapshots/Private/Data/Util/SnapshotUtil.h
--- a/Engine/Plugins/VirtualProduction/LevelSnapshots/Source/LevelSnapshots/Private/Data/Util/SnapshotUtil.h
+++ b/Engine/Plugins/VirtualProduction/LevelSnapshots/Source/LevelSnapshots/Private/Data/Util/SnapshotUtil.h
@@ -10,6 +10,9 @@ struct FActorSnapshotData;
 
 namespace UE::LevelSnapshots::Private
 {
+	/** Given e.g. "/Game/MapName.MapName:PersistentLevel.StaticMeshActor_42.SomeComponent" returns everything before the colon, i.e. "/Game/MapName.MapName"*/
+	FSoftObjectPath ExtractPathWithoutSubobjects(const FSoftObjectPath& ObjectPath);
+	
 	/** Finds the last subobject name in the path */
 	FString ExtractLastSubobjectName(const FSoftObjectPath& ObjectPath);
 
diff --git a/Engine/Plugins/VirtualProduction/LevelSnapshots/Source/LevelSnapshots/Private/Data/Util/WorldData/AllocationActorUtil.cpp b/Engine/Plugins/VirtualProduction/LevelSnapshots/Source/LevelSnapshots/Private/Data/Util/WorldData/AllocationActorUtil.cpp
new file mode 100644
--- /dev/null
+++ b/Engine/Plugins/VirtualProduction/LevelSnapshots/Source/LevelSnapshots/Private/Data/Util/WorldData/AllocationActorUtil.cpp
@@ -0,0 +1,372 @@
+﻿// Copyright Epic Games, Inc. All Rights Reserved.
+
+#include "Data/Util/WorldData/ActorUtil.h"
+
+#include "ClassDataUtil.h"
+#include "Data/ActorSnapshotData.h"
+#include "Data/WorldSnapshotData.h"
+#include "Data/SnapshotCustomVersion.h"
+#include "CustomSerialization/CustomObjectSerializationWrapper.h"
+#include "Data/Util/WorldData/SnapshotObjectUtil.h"
+#include "LevelSnapshotsLog.h"
+#include "LevelSnapshotsModule.h"
+#include "LoadSnapshotObjectArchive.h"
+#include "SnapshotDataCache.h"
+#include "Selection/PropertySelectionMap.h"
+#include "Util/EquivalenceUtil.h"
+#include "Util/Component/SnapshotComponentUtil.h"
+
+#include "Components/ActorComponent.h"
+#include "Engine/World.h"
+#include "GameFramework/Actor.h"
+#include "Templates/NonNullPointer.h"
+#include "UObject/Package.h"
+#include "UObject/Script.h"
+#include "Util/SnapshotUtil.h"
+
+#if USE_STABLE_LOCALIZATION_KEYS
+#include "Internationalization/TextPackageNamespaceUtil.h"
+#endif
+
+namespace UE::LevelSnapshots::Private::Internal
+{
+	static void PostSerializeSnapshotActor(AActor* SnapshotActor, const FSoftObjectPath& OriginalActorPath, FWorldSnapshotData& WorldData, FSnapshotDataCache& Cache, UPackage* InLocalisationSnapshotPackage)
+	{
+		struct Local
+		{
+			static void RecreateRootComponentIfInstanced(AActor* OriginalActor, AActor* SnapshotActor)
+			{
+				USceneComponent* OriginalRoot = OriginalActor->GetRootComponent();
+				
+				const bool bSnapshotIsMissingRoot = SnapshotActor->GetRootComponent() == nullptr;
+				// Some actors, e.g. exact instances of AActor, have no components and the editor creates an instanced root component
+				const bool bSnapshotWasSavedWithoutRootComponentData = OriginalRoot && OriginalRoot->CreationMethod == EComponentCreationMethod::Instance;
+				if (bSnapshotIsMissingRoot && bSnapshotWasSavedWithoutRootComponentData)
+				{
+					USceneComponent* SnapshotRoot = DuplicateObject<USceneComponent>(OriginalRoot, SnapshotActor, OriginalRoot->GetFName());
+					UE_CLOG(SnapshotRoot == nullptr, LogLevelSnapshots, Warning, TEXT("Failed to recreate instanced root component for %s's snapshot counterpart"), *OriginalActor->GetPathName());
+					if (SnapshotRoot)
+					{
+						SnapshotActor->SetRootComponent(
+							SnapshotRoot
+						);
+					}
+				}
+			}
+			
+			static void HandleAttachParentNotSaved(AActor* OriginalActor, AActor* SnapshotActor, FWorldSnapshotData& WorldData, FSnapshotDataCache& Cache, UPackage* InLocalisationSnapshotPackage)
+			{
+				USceneComponent* RootComponent = OriginalActor->GetRootComponent();
+				// WorldSettings does not have a root component...
+				if (!RootComponent)
+				{
+					return;
+				}
+				
+				USceneComponent* WorldAttachParent = RootComponent->GetAttachParent();
+				if (!WorldAttachParent)
+				{
+					return;
+				}
+
+				// We're covering the case of when the attach parent was not saved... if the current attach parent was not saved... we'll create our own warning
+				constexpr bool bWarnIfNotFound = false;
+				const TOptional<TNonNullPtr<AActor>> SnapshotAttachParent = GetDeserializedActor(WorldAttachParent->GetOwner(), WorldData, Cache, InLocalisationSnapshotPackage, bWarnIfNotFound);
+				if (!SnapshotAttachParent)
+				{
+					UE_LOG(LogLevelSnapshots, Warning,
+						TEXT("Actor %s had no attach parent saved. Its current attach parent %s was added after the snapshot was added so we when this actor is restored, we'll fallback to unparenting it."),
+						*OriginalActor->GetPathName(),
+						*WorldAttachParent->GetOwner()->GetPathName());
+					return;
+				}
+
+				const TInlineComponentArray<USceneComponent*> ComponentArray(SnapshotAttachParent.GetValue());
+				USceneComponent* const* SnapshotAttachComponent = ComponentArray.FindByPredicate([WorldAttachParent](USceneComponent* Comp)
+				{
+					return Comp->GetFName() == WorldAttachParent->GetFName();
+				});
+
+				if (SnapshotAttachComponent)
+				{
+					SnapshotActor->AttachToComponent(
+						*SnapshotAttachComponent,
+						// Do not change property values: anything other than KeepRelativeTransform will recompute
+						FAttachmentTransformRules::KeepRelativeTransform,
+						RootComponent->GetAttachSocketName()
+						);
+				}
+				else
+				{
+					UE_LOG(LogLevelSnapshots, Warning, TEXT("Failed to fix up AttachParent for %s's snapshot counterpart"), *OriginalActor->GetPathName());
+				}
+			}
+		};
+		
+		const int32 SavedVersion = WorldData.SnapshotVersionInfo.GetSnapshotCustomVersion();
+		AActor* OriginalActor = Cast<AActor>(OriginalActorPath.ResolveObject());
+
+		// USceneComponent::AttachParent was not yet saved.
+		if (SavedVersion < FSnapshotCustomVersion::SubobjectSupport
+			&& OriginalActor)
+		{
+			Local::RecreateRootComponentIfInstanced(OriginalActor, SnapshotActor);
+			// To avoid lots of actors being unparented by accident, we set the snapshot actor's attach parent to the equivalent of the editor world
+			Local::HandleAttachParentNotSaved(OriginalActor, SnapshotActor, WorldData, Cache, InLocalisationSnapshotPackage);
+		}
+	}
+
+	using FHandleFoundComponent = TFunctionRef<void(FSubobjectSnapshotData& SerializedCompData, UActorComponent* ActorComp, const FSoftObjectPath& OriginalComponentPath, FWorldSnapshotData& SharedData)>;
+
+	static void DeserializeComponents(
+		const FActorSnapshotData& ActorData,
+		AActor* IntoActor,
+		FWorldSnapshotData& WorldData,
+		FHandleFoundComponent HandleComponent
+	)
+	{
+		for (auto CompIt = ActorData.ComponentData.CreateConstIterator(); CompIt; ++CompIt)
+		{
+			const EComponentCreationMethod CreationMethod = CompIt->Value.CreationMethod;
+			if (CreationMethod == EComponentCreationMethod::UserConstructionScript)	// Construction script components are not supported 
+			{
+				continue;
+			}
+
+			const int32 SubobjectIndex = CompIt->Key;
+			const FSoftObjectPath& OriginalComponentPath = WorldData.SerializedObjectReferences[SubobjectIndex];
+			if (UActorComponent* ComponentToRestore = FindMatchingComponent(IntoActor, OriginalComponentPath))
+			{
+				FSubobjectSnapshotData& SnapshotData = WorldData.Subobjects[SubobjectIndex];
+				HandleComponent(SnapshotData, ComponentToRestore, OriginalComponentPath, WorldData);
+			}
+		}
+	}
+
+	static void DeserializeSubobjectsForSnapshotActor(const FActorSnapshotData& ActorData, FWorldSnapshotData& WorldData, FSnapshotDataCache& Cache, const FProcessObjectDependency& ProcessObjectDependency, UPackage* InLocalisationSnapshotPackage)
+	{
+		for (const int32 SubobjectIndex : ActorData.OwnedSubobjects)
+		{
+			check(!ActorData.ComponentData.Contains(SubobjectIndex));
+
+			FString LocalisationNamespace;
+#if USE_STABLE_LOCALIZATION_KEYS
+			LocalisationNamespace = TextNamespaceUtil::EnsurePackageNamespace(InLocalisationSnapshotPackage);
+#endif
+
+			// Ensures the object is allocated and serialized into. Return value not needed.
+			ResolveObjectDependencyForSnapshotWorld(
+				WorldData,
+				Cache,
+				SubobjectIndex,
+				ProcessObjectDependency,
+				LocalisationNamespace
+			);
+		}
+	}
+
+	static void DeserializeIntoTransient(
+		FSubobjectSnapshotData& SerializedComponentData,
+		UActorComponent* ComponentToDeserializeInto,
+		FWorldSnapshotData& WorldData,
+		const FProcessObjectDependency& ProcessObjectDependency,
+		FSnapshotDataCache& Cache,
+		UPackage* InLocalisationSnapshotPackage)
+	{
+		if (ComponentToDeserializeInto->CreationMethod == EComponentCreationMethod::UserConstructionScript)
+		{
+			UE_LOG(LogLevelSnapshots, Warning, TEXT("Components created dynamically in the construction script are not supported (%s). Skipping..."), *ComponentToDeserializeInto->GetPathName());
+			return;
+		}
+		
+		// We cannot supply a template object for components whose CreationMethod is not Instance.
+		// Yet every actor class, native and Blueprint, has their own archetype for each component, i.e. in the details panel the property has a yellow "Reset to Default" icon next to it if the value is different from the Blueprint's value.
+		// We apply the CDO data we saved before we apply the component's serialized data.
+			// Scenario: We have a custom Blueprint StaticMeshBP that has a UStaticMeshComponent. Suppose we change some default property value, like ComponentTags.
+			// 1st case: We change no default values, then take a snapshot.
+			//   - After the snapshot was taken, we change a default value.
+			//   - The value of the property was not saved in the component's serialized data because the value was equal to the CDO's value.
+			//   - When applying the snapshot, we will override the new default value because we serialize the CDO here. Good.
+			// 2nd case: We change some default value, then take a snapshot
+			//   - The value was saved into the component's serialized data because it was different from the CDO's default value at the time of taking the snapshot
+			//   - We apply the CDO and afterwards we override it with the serialized data. Good.
+		const FSubobjectArchetypeFallbackInfo Fallback{ ComponentToDeserializeInto->GetOuter(), ComponentToDeserializeInto->GetFName(), SerializedComponentData.GetObjectFlags() };
+		SerializeClassDefaultsIntoSubobject(ComponentToDeserializeInto, WorldData, SerializedComponentData.ClassIndex, Cache, Fallback);
+		
+		FLoadSnapshotObjectArchive::ApplyToSnapshotWorldObject(SerializedComponentData, WorldData, Cache, ComponentToDeserializeInto, ProcessObjectDependency, InLocalisationSnapshotPackage);
+	}
+
+	static TOptional<TNonNullPtr<AActor>> GetDeserialized(const FSoftObjectPath& OriginalActorPath, FWorldSnapshotData& WorldData, FSnapshotDataCache& Cache, UPackage* InLocalisationSnapshotPackage, bool bWarnIfNotFound)
+	{
+		SCOPED_SNAPSHOT_CORE_TRACE(GetDeserialized);
+		UE_LOG(LogLevelSnapshots, Verbose, TEXT("========== Get Deserialized %s =========="), *OriginalActorPath.ToString());
+		
+		const TOptional<TNonNullPtr<AActor>> Preallocated = GetPreallocated(OriginalActorPath, WorldData, Cache, bWarnIfNotFound);
+		if (!Preallocated)
+		{
+			return {};
+		}
+		
+		FActorSnapshotCache& ActorCache = Cache.ActorCache.FindOrAdd(OriginalActorPath);
+		ActorCache.bReceivedSerialisation = true;
+		
+		FActorSnapshotData& ActorData = WorldData.ActorData[OriginalActorPath];
+		const auto ProcessObjectDependency = [&Cache, &WorldData, &OriginalActorPath](int32 OriginalObjectDependency)
+		{
+			check(WorldData.SerializedObjectReferences.IsValidIndex(OriginalObjectDependency));
+			
+			// This look-up must be done every time because FSnapshotDataCache::ActorCache may have been reallocated 
+			FActorSnapshotCache& ActorCache = Cache.ActorCache.FindOrAdd(OriginalActorPath);
+			ActorCache.ObjectDependencies.Add(OriginalObjectDependency);
+		};
+		
+		AActor* PreallocatedActor = Preallocated.GetValue();
+		{
+			const FRestoreObjectScope FinishRestore = PreActorRestore_SnapshotWorld(PreallocatedActor, ActorData.CustomActorSerializationData, WorldData, Cache, ProcessObjectDependency, InLocalisationSnapshotPackage);
+			FLoadSnapshotObjectArchive::ApplyToSnapshotWorldObject(ActorData.SerializedActorData, WorldData, Cache, PreallocatedActor, ProcessObjectDependency, InLocalisationSnapshotPackage);
+#if WITH_EDITOR
+			UE_LOG(LogLevelSnapshots, Verbose, TEXT("ActorLabel is \"%s\" for \"%s\" (editor object path \"%s\")"), *PreallocatedActor->GetActorLabel(), *PreallocatedActor->GetPathName(), *OriginalActorPath.ToString());
+#endif
+		}
+		
+		DeserializeComponents(ActorData, PreallocatedActor, WorldData,
+			[&WorldData, &ProcessObjectDependency, &Cache, InLocalisationSnapshotPackage](
+				FSubobjectSnapshotData& SerializedCompData,
+				UActorComponent* Comp,
+				const FSoftObjectPath& OriginalComponentPath,
+				FWorldSnapshotData& SharedData)
+			{
+				FSubobjectSnapshotCache& SubobjectCache = Cache.SubobjectCache.FindOrAdd(OriginalComponentPath);
+				SubobjectCache.SnapshotObject = Comp;
+				
+				const FRestoreObjectScope FinishRestore = PreSubobjectRestore_SnapshotWorld(Comp, OriginalComponentPath, WorldData, Cache , ProcessObjectDependency, InLocalisationSnapshotPackage);
+				DeserializeIntoTransient(SerializedCompData, Comp, SharedData, ProcessObjectDependency, Cache, InLocalisationSnapshotPackage);
+			}
+		);
+
+		DeserializeSubobjectsForSnapshotActor(ActorData, WorldData, Cache, ProcessObjectDependency, InLocalisationSnapshotPackage);
+		PostSerializeSnapshotActor(PreallocatedActor, OriginalActorPath, WorldData, Cache, InLocalisationSnapshotPackage);
+		PreallocatedActor->UpdateComponentTransforms();
+		{
+			// GAllowActorScriptExecutionInEditor must be temporarily true so call to UserConstructionScript isn't skipped
+			FEditorScriptExecutionGuard AllowConstructionScript;
+			PreallocatedActor->UserConstructionScript();
+		}
+		
+		return Preallocated;
+	}
+	
+	static void ConditionallyRerunConstructionScript(AActor* RequiredActor, const TArray<int32>& OriginalObjectDependencies, FWorldSnapshotData& WorldData, FSnapshotDataCache& Cache, UPackage* LocalisationSnapshotPackage)
+	{
+		if (!RequiredActor)
+		{
+			return;
+		}
+		
+		bool bHadActorDependencies = false;
+		for (const int32 OriginalObjectIndex : OriginalObjectDependencies)
+		{
+			const FSoftObjectPath& ObjectPath = WorldData.SerializedObjectReferences[OriginalObjectIndex];
+			if (WorldData.ActorData.Contains(ObjectPath))
+			{
+				GetDeserializedActor(ObjectPath, WorldData, Cache, LocalisationSnapshotPackage);
+				bHadActorDependencies = true;
+			}
+		}
+		
+		if (bHadActorDependencies && ensure(RequiredActor))
+		{
+			// GAllowActorScriptExecutionInEditor must be temporarily true so call to UserConstructionScript isn't skipped
+			FEditorScriptExecutionGuard AllowConstructionScript;
+			RequiredActor->UserConstructionScript();
+		}
+	}
+
+	static ULevel* DetermineMostSuitableLevel(const FSoftObjectPath& ActorPath, const FWorldSnapshotData& WorldData)
+	{
+		const FName LevelName = *ExtractPathWithoutSubobjects(ActorPath).GetAssetName();
+		for (const UWorld* SublevelWorld : WorldData.SnapshotSublevels)
+		{
+			if (SublevelWorld->GetFName() == LevelName)
+			{
+				return SublevelWorld->PersistentLevel;
+			}
+		}
+		return WorldData.SnapshotWorld->PersistentLevel;
+	}
+}
+
+TOptional<TNonNullPtr<AActor>> UE::LevelSnapshots::Private::GetDeserializedActor(const FSoftObjectPath& OriginalObjectPath, FWorldSnapshotData& WorldData, FSnapshotDataCache& Cache, UPackage* LocalisationSnapshotPackage, bool bWarnIfNotFound)
+{
+	{
+		// Note: CacheBeforeDeserialize may become invalid after the GetDeserialized call
+		const FActorSnapshotCache& CacheBeforeDeserialize = Cache.ActorCache.FindOrAdd(OriginalObjectPath);
+		if (CacheBeforeDeserialize.bReceivedSerialisation && CacheBeforeDeserialize.CachedSnapshotActor.IsValid())
+		{
+			return { CacheBeforeDeserialize.CachedSnapshotActor.Get() };
+		}
+	}
+	
+	const TOptional<TNonNullPtr<AActor>> Result = Internal::GetDeserialized(OriginalObjectPath, WorldData, Cache, LocalisationSnapshotPackage, bWarnIfNotFound);
+	// Look up cache again because GetDeserialized might have reallocated Cache.ActorCache hence CacheBeforeDeserialize is no longer valid
+	const FActorSnapshotCache& CacheAfterDeserialize = Cache.ActorCache.FindOrAdd(OriginalObjectPath);
+	Internal::ConditionallyRerunConstructionScript(Result.Get(nullptr), CacheAfterDeserialize.ObjectDependencies, WorldData, Cache, LocalisationSnapshotPackage);
+	return Result;
+}
+
+TOptional<TNonNullPtr<AActor>> UE::LevelSnapshots::Private::GetPreallocatedIfCached(const FSoftObjectPath& OriginalActorPath, const FSnapshotDataCache& Cache)
+{
+	if (const FActorSnapshotCache* ActorCache = Cache.ActorCache.Find(OriginalActorPath))
+	{
+		return ActorCache->CachedSnapshotActor.IsValid() ? TOptional<TNonNullPtr<AActor>>(ActorCache->CachedSnapshotActor.Get()) : TOptional<TNonNullPtr<AActor>>();
+	}
+	return {};
+}
+
+TOptional<TNonNullPtr<AActor>> UE::LevelSnapshots::Private::GetPreallocated(const FSoftObjectPath& OriginalActorPath, FWorldSnapshotData& WorldData, FSnapshotDataCache& Cache, bool bWarnIfNotFound)
+{
+	SCOPED_SNAPSHOT_CORE_TRACE(GetPreallocated);
+	FActorSnapshotCache& ActorCache = Cache.ActorCache.FindOrAdd(OriginalActorPath);
+	if (!ActorCache.CachedSnapshotActor.IsValid())
+	{
+		UWorld* SnapshotWorld = WorldData.SnapshotWorld->GetWorld();
+		const FActorSnapshotData* ActorData = WorldData.ActorData.Find(OriginalActorPath);
+		if (!ActorData)
+		{
+			UE_CLOG(bWarnIfNotFound, LogLevelSnapshots, Warning, TEXT("Failed to allocate %s due to missing data"), *OriginalActorPath.ToString());
+			return {};
+		}
+		
+		const TOptional<TNonNullPtr<AActor>> Template = GetActorClassDefault(WorldData, ActorData->ClassIndex, Cache);
+		if (!Template)
+		{
+			UE_LOG(LogLevelSnapshots, Warning, TEXT("Failed to allocate %s due to missing class."), *OriginalActorPath.ToString());
+			return {};
+		}
+
+		UClass* TargetClass = Template->GetClass();
+		FActorSpawnParameters SpawnParams;
+		SpawnParams.Template = Template.GetValue();
+		SpawnParams.NameMode = FActorSpawnParameters::ESpawnActorNameMode::Requested;
+		SpawnParams.OverrideLevel = Internal::DetermineMostSuitableLevel(OriginalActorPath, WorldData);
+		const FName ActorName = *ExtractLastSubobjectName(OriginalActorPath);
+		// Must still generate a unique name because the asset registry might have messed with our data during a rename and added a duplicate actor name...
+		SpawnParams.Name = FindObjectFast<UObject>(SpawnParams.OverrideLevel, ActorName)
+			? MakeUniqueObjectName(SpawnParams.OverrideLevel, TargetClass, ActorName)
+			: ActorName;
+		ActorCache.CachedSnapshotActor = SnapshotWorld->SpawnActor<AActor>(TargetClass, SpawnParams);
+
+		const bool bSpawnWasSuccess = ActorCache.CachedSnapshotActor.IsValid();
+		UE_CLOG(!bSpawnWasSuccess, LogLevelSnapshots, Error, TEXT("Failed to allocate %s due to missing class."), *OriginalActorPath.ToString());
+		if (ensure(bSpawnWasSuccess))
+		{
+#if WITH_EDITOR
+			// Hide this actor so external systems can see that this components should not render, i.e. make USceneComponent::ShouldRender return false
+			ActorCache.CachedSnapshotActor->SetIsTemporarilyHiddenInEditor(true);
+#endif
+			AllocateMissingComponentsForSnapshotActor(ActorCache.CachedSnapshotActor.Get(), OriginalActorPath, WorldData, Cache);
+		}
+	}
+
+	return ActorCache.CachedSnapshotActor.IsValid() ? TOptional<TNonNullPtr<AActor>>(ActorCache.CachedSnapshotActor.Get()) : TOptional<TNonNullPtr<AActor>>();
+}
\ No newline at end of file
