Merge changes from Otter stream to Main

diff --git a/Engine/Plugins/VirtualProduction/LevelSnapshots/Source/LevelSnapshots/Private/Data/Util/WorldData/ClassDataUtil.cpp b/Engine/Plugins/VirtualProduction/LevelSnapshots/Source/LevelSnapshots/Private/Data/Util/WorldData/ClassDataUtil.cpp
new file mode 100644
--- /dev/null
+++ b/Engine/Plugins/VirtualProduction/LevelSnapshots/Source/LevelSnapshots/Private/Data/Util/WorldData/ClassDataUtil.cpp
@@ -0,0 +1,190 @@
+// Copyright Epic Games, Inc. All Rights Reserved.
+
+#include "ClassDataUtil.h"
+
+#include "ActorSnapshotData.h"
+#include "ClassDefaults/ApplyClassDefaulDataArchive.h"
+#include "ClassDefaults/TakeClassDefaultObjectSnapshotArchive.h"
+#include "Data/WorldSnapshotData.h"
+#include "LevelSnapshotsLog.h"
+#include "LevelSnapshotsModule.h"
+#include "SnapshotUtilTypes.h"
+#include "SubobjectSnapshotData.h"
+
+#include "Algo/ForEach.h"
+#include "EngineUtils.h"
+#include "SnapshotDataCache.h"
+#include "Templates/NonNullPointer.h"
+#include "UObject/Package.h"
+
+namespace UE::LevelSnapshots::Private
+{
+	TOptional<TNonNullPtr<AActor>> GetActorClassDefault(FWorldSnapshotData& WorldData, FClassDataIndex ClassIndex, FSnapshotDataCache& Cache)
+	{
+		check(WorldData.ClassData.IsValidIndex(ClassIndex));
+		
+		FClassSnapshotData& ClassData = WorldData.ClassData[ClassIndex];
+		UClass* Class = ClassData.ClassPath.TryLoadClass<AActor>();
+		if (!Class)
+		{
+			UE_LOG(LogLevelSnapshots, Error, TEXT("Unknown class %s. The snapshot is mostly likely referencing a class that was deleted."), *ClassData.ClassPath.ToString());
+			return {};
+		}
+		
+		if (ClassData.SnapshotFlags == ESnapshotClassFlags::SerializationSkippedArchetypeData
+			|| ClassData.SerializedData.Num() == 0)
+		{
+			return { CastChecked<AActor>(Class->GetDefaultObject()) };
+		}
+
+		TObjectPtr<UObject>& CachedCDO = Cache.ArchetypeObjects[ClassIndex];
+		if (IsValid(CachedCDO))
+		{
+			return { CastChecked<AActor>(CachedCDO) };
+		}
+		
+		UObject* CDO = NewObject<UObject>(
+			GetTransientPackage(),
+			Class,
+			*FString("SnapshotCDO_").Append(*MakeUniqueObjectName(GetTransientPackage(), Class).ToString())
+			// TODO: Set RF_Archetype
+			);
+		FApplyClassDefaulDataArchive::SerializeClassDefaultObject(ClassData, WorldData, CDO);
+
+		CachedCDO = CDO;
+		return { CastChecked<AActor>(CDO) };
+	}
+
+	TOptional<TNonNullPtr<UObject>> GetSubobjectArchetype(FWorldSnapshotData& WorldData, FClassDataIndex ClassIndex, FSnapshotDataCache& Cache, const FSubobjectArchetypeFallbackInfo& FallbackInfo)
+	{
+		if (Cache.ArchetypeObjects[ClassIndex])
+		{
+			return { Cache.ArchetypeObjects[ClassIndex].Get() };
+		}
+		
+		const TOptional<TNonNullPtr<FClassSnapshotData>> ArchetypeData = GetSubobjectArchetypeData(WorldData, ClassIndex, Cache, FallbackInfo);
+		if (!ArchetypeData)
+		{
+			return {};
+		}
+
+		UClass* Class = ArchetypeData->ClassPath.TryLoadClass<UObject>();
+		UObject* Archetype = NewObject<UObject>(
+			GetTransientPackage(),
+			Class,
+			*FString("SnapshotArchetype_").Append(*MakeUniqueObjectName(GetTransientPackage(), Class).ToString())
+			// TODO: Set RF_Archetype
+			);
+		FApplyClassDefaulDataArchive::SerializeClassDefaultObject(*ArchetypeData.GetValue(), WorldData, Archetype);
+		Cache.ArchetypeObjects[ClassIndex] = Archetype;
+		
+		return { Archetype };
+	}
+
+	TOptional<TNonNullPtr<FClassSnapshotData>> GetSubobjectArchetypeData(FWorldSnapshotData& WorldData, FClassDataIndex ClassIndex, FSnapshotDataCache& Cache, const FSubobjectArchetypeFallbackInfo& FallbackInfo)
+	{
+		if (!ensureMsgf(WorldData.ClassData.IsValidIndex(ClassIndex), TEXT("This data was supposed to be saved when the snapshot was taken. Data is corrupt.")))
+		{
+			UE_LOG(LogLevelSnapshots, Error, TEXT("Snapshot data is corrupt. ClassIndex %d is out of bounds for ClassData of size %d."), ClassIndex, WorldData.ClassData.Num());
+			return {};
+		}
+
+		// Normal case: the data was saved
+		FClassSnapshotData& SavedClassData = WorldData.ClassData[ClassIndex];
+		const bool bHasArchetypeData = SavedClassData.SerializedData.Num() != 0
+			&& (SavedClassData.SnapshotFlags & ESnapshotClassFlags::SerializationSkippedArchetypeData) == ESnapshotClassFlags::NoFlags;
+		if (bHasArchetypeData)
+		{
+			return { &SavedClassData };
+		}
+
+		// This is an old snapshot (or snapshot config prevented data from being saved)
+			// 1. Since serialization is expensive, look it up in the fallback cache
+		TOptional<FClassSnapshotData>& FallbackData = Cache.FallbackArchetypeData[ClassIndex];
+		if (FallbackData.IsSet())
+		{
+			return { &FallbackData.GetValue() };
+		}
+
+		// 2. First time encountering this class, get the real archetype
+		UObject* RealArchetype = UObject::GetArchetypeFromRequiredInfo(SavedClassData.ClassPath.TryLoadClass<UObject>(), FallbackInfo.SubobjectOuter, FallbackInfo.SubobjectName, FallbackInfo.SubobjectFlags);
+		if (ensure(RealArchetype))
+		{
+			FallbackData = FClassSnapshotData(); // Optional needs to be assigned
+			FTakeClassDefaultObjectSnapshotArchive::SaveClassDefaultObject(FallbackData.GetValue(), WorldData, RealArchetype);
+			return { &FallbackData.GetValue() };
+		}
+
+		UE_LOG(LogLevelSnapshots, Error, TEXT("Failed to determine archetype for class: %s, outer: %s name: %s"), *SavedClassData.ClassPath.ToString(), *FallbackInfo.SubobjectOuter->GetPathName(), *FallbackInfo.SubobjectName.ToString())
+		return {};
+	}
+
+	void SerializeClassDefaultsIntoSubobject(UObject* Object, FWorldSnapshotData& WorldData, FClassDataIndex ClassIndex, FSnapshotDataCache& Cache, const FSubobjectArchetypeFallbackInfo& FallbackInfo)
+	{
+		if (const TOptional<TNonNullPtr<FClassSnapshotData>> ArchetypeData = GetSubobjectArchetypeData(WorldData, ClassIndex, Cache, FallbackInfo))
+		{
+			SerializeClassDefaultsIntoSubobject(Object, *ArchetypeData, WorldData);
+		}
+		else
+		{
+			UE_LOG(LogLevelSnapshots, Warning,
+				TEXT("No CDO saved for class '%s'. If you changed some class default values for this class, then the affected objects will have the latest values instead of the class defaults at the time the snapshot was taken. Should be nothing major to worry about."),
+				*Object->GetClass()->GetName()
+				);
+		}
+	}
+
+	void SerializeClassDefaultsIntoSubobject(UObject* Object, FClassSnapshotData& DataToSerialize, FWorldSnapshotData& WorldData)
+	{
+		const bool bHasData = DataToSerialize.SerializedData.Num() > 0 && (DataToSerialize.SnapshotFlags & ESnapshotClassFlags::SerializationSkippedArchetypeData) == ESnapshotClassFlags::NoFlags;
+		if (bHasData && !FLevelSnapshotsModule::GetInternalModuleInstance().ShouldSkipClassDefaultSerialization(Object->GetClass()))
+		{
+			FApplyClassDefaulDataArchive::RestoreChangedClassDefaults(DataToSerialize, WorldData, Object);
+		}
+	}
+
+	FSoftClassPath GetClass(const FActorSnapshotData& Data, const FWorldSnapshotData& WorldData)
+	{
+		return WorldData.ClassData[Data.ClassIndex].ClassPath;
+	}
+	
+	FSoftClassPath GetClass(const FSubobjectSnapshotData& Data, const FWorldSnapshotData& WorldData)
+	{
+		return WorldData.ClassData[Data.ClassIndex].ClassPath;
+	}
+	
+	FClassDataIndex AddClassArchetype(FWorldSnapshotData& WorldData, UObject* SavedObject)
+	{
+		check(SavedObject);
+		checkf(!SavedObject->IsA<UClass>(), TEXT("You are supposed to put in the original object, not its class."));
+
+		UObject* Archetype = SavedObject->GetArchetype();
+		if (const FClassDataIndex* Index = WorldData.ArchetypeToClassDataIndex.Find(Archetype))
+		{
+			return *Index;
+		}
+		
+		UClass* Class = SavedObject->GetClass();
+		FClassSnapshotData ClassData;
+		ClassData.ClassPath = Class;
+		ClassData.ClassFlags = Class->GetFlags();
+		// Reserve an index and init cache now because we call SaveClassDefaultObject below
+		// which by calling AddClassArchetype can reallocate ClassData and ClassToClassDataIndex...
+		const int32 Index = WorldData.ClassData.Add(ClassData);
+		WorldData.ArchetypeToClassDataIndex.Add(Archetype, Index);
+		
+		const bool bShouldSkip = FLevelSnapshotsModule::GetInternalModuleInstance().ShouldSkipClassDefaultSerialization(Class);
+		if (bShouldSkip)
+		{
+			ClassData.SnapshotFlags = ESnapshotClassFlags::SerializationSkippedArchetypeData;
+		}
+		else
+		{
+			FTakeClassDefaultObjectSnapshotArchive::SaveClassDefaultObject(ClassData, WorldData, Archetype);
+		}
+		
+		// ... now we can write the serialized data safely
+		WorldData.ClassData[Index] = MoveTemp(ClassData);
+		return Index;
+	}
+}
\ No newline at end of file
diff --git a/Engine/Plugins/VirtualProduction/LevelSnapshots/Source/LevelSnapshots/Private/Data/Util/WorldData/RestoreActorUtil.cpp b/Engine/Plugins/VirtualProduction/LevelSnapshots/Source/LevelSnapshots/Private/Data/Util/WorldData/RestoreActorUtil.cpp
new file mode 100644
--- /dev/null
+++ b/Engine/Plugins/VirtualProduction/LevelSnapshots/Source/LevelSnapshots/Private/Data/Util/WorldData/RestoreActorUtil.cpp
@@ -0,0 +1,215 @@
+﻿// Copyright Epic Games, Inc. All Rights Reserved.
+
+#include "ApplySnapshotToEditorArchive.h"
+
+#include "CustomSerialization/CustomObjectSerializationWrapper.h"
+#include "Data/ActorSnapshotData.h"
+#include "Data/RestorationEvents/ApplySnapshotToActorScope.h"
+#include "Data/Util/EquivalenceUtil.h"
+#include "Data/Util/Component/SnapshotComponentUtil.h"
+#include "Data/Util/WorldData/ActorUtil.h"
+#include "Data/WorldSnapshotData.h"
+#include "LevelSnapshotsLog.h"
+#include "Selection/PropertySelectionMap.h"
+
+#include "Components/ActorComponent.h"
+#include "Engine/Engine.h"
+#include "GameFramework/Actor.h"
+#include "Templates/NonNullPointer.h"
+#include "UObject/Package.h"
+
+#if USE_STABLE_LOCALIZATION_KEYS
+#include "Internationalization/TextPackageNamespaceUtil.h"
+#endif
+
+namespace UE::LevelSnapshots::Private::Internal::Restore
+{
+	using FSerializeActor = TFunctionRef<void(AActor* OriginalActor, AActor* DerserializedActor)>;
+	using FSerializeComponent = TFunctionRef<void(FSubobjectSnapshotData& SerializedCompData, UActorComponent* Original, UActorComponent* Deserialized)>;
+	using FHandleFoundComponent = TFunctionRef<void(FSubobjectSnapshotData& SerializedCompData, UActorComponent* ActorComp)>;
+
+	static void DeserializeComponents(
+		AActor* IntoActor,
+		const FActorSnapshotData& ActorData,
+		FWorldSnapshotData& WorldData,
+		FHandleFoundComponent HandleComponent
+		)
+	{
+		for (auto CompIt = ActorData.ComponentData.CreateConstIterator(); CompIt; ++CompIt)
+		{
+			const EComponentCreationMethod CreationMethod = CompIt->Value.CreationMethod;
+			if (CreationMethod == EComponentCreationMethod::UserConstructionScript)	// Construction script components are not supported 
+				{
+				continue;
+				}
+
+			const int32 SubobjectIndex = CompIt->Key;
+			const FSoftObjectPath& OriginalComponentPath = WorldData.SerializedObjectReferences[SubobjectIndex];
+			if (UActorComponent* ComponentToRestore = UE::LevelSnapshots::Private::FindMatchingComponent(IntoActor, OriginalComponentPath))
+			{
+				FSubobjectSnapshotData& SnapshotData = WorldData.Subobjects[SubobjectIndex];
+				HandleComponent(SnapshotData, ComponentToRestore);
+			}
+		}
+	}
+	
+	static void DeserializeIntoEditorWorldActor(AActor* OriginalActor, FActorSnapshotData& ActorData, FWorldSnapshotData& WorldData, FSnapshotDataCache& Cache, UPackage* InLocalisationSnapshotPackage, FSerializeActor SerializeActor, FSerializeComponent SerializeComponent)
+	{
+		const TOptional<TNonNullPtr<AActor>> Deserialized = UE::LevelSnapshots::Private::GetDeserializedActor(OriginalActor, WorldData, Cache, InLocalisationSnapshotPackage);
+		if (!Deserialized)
+		{
+			UE_LOG(LogLevelSnapshots, Warning, TEXT("Failed to serialize into actor %s. Skipping..."), *OriginalActor->GetName());
+			return;
+		}
+
+		const AActor* AttachParentBeforeRestore = OriginalActor->GetAttachParentActor();
+		SerializeActor(OriginalActor, Deserialized.GetValue());
+#if WITH_EDITOR
+		UE_LOG(LogLevelSnapshots, Verbose, TEXT("ActorLabel is \"%s\" for \"%s\""), *OriginalActor->GetActorLabel(), *OriginalActor->GetPathName());
+#endif
+
+		TInlineComponentArray<UActorComponent*> DeserializedComponents(Deserialized.GetValue());
+		Internal::Restore::DeserializeComponents(OriginalActor, ActorData, WorldData,
+			[&SerializeComponent, &DeserializedComponents](
+				FSubobjectSnapshotData& SerializedCompData,
+				UActorComponent* Comp
+				)
+		    {
+		        const FName OriginalCompName = Comp->GetFName();
+		        UActorComponent** DeserializedCompCounterpart = DeserializedComponents.FindByPredicate([OriginalCompName](UActorComponent* Other)
+		        {
+		            return Other->GetFName() == OriginalCompName;
+		        });
+				
+				UE_CLOG(!DeserializedCompCounterpart, LogLevelSnapshots, Warning, TEXT("Failed to find component called %s on temp deserialized snapshot actor. Skipping component..."), *OriginalCompName.ToString())
+		        if (DeserializedCompCounterpart)
+		        {
+	        		SerializeComponent(SerializedCompData, Comp, *DeserializedCompCounterpart);
+		        }
+		    }
+		);
+
+		// Handle case in which AttachParentAfterRestore is restored after OriginalActor otherwise ReregisterAllComponents will detach OriginalActor again
+		const AActor* AttachParentAfterRestore = OriginalActor->GetAttachParentActor();
+		if (AttachParentAfterRestore && AttachParentAfterRestore->IsAttachedTo(OriginalActor))
+		{
+			AttachParentAfterRestore->GetRootComponent()->DetachFromComponent(FDetachmentTransformRules::KeepRelativeTransform);
+		}
+
+		// Update component state, e.g. render state if intensity for lights was changed. Also avoids us having to call PostEditChangeProperty on every property.
+		// ReregisterAllComponents must be called before the UserConstructionScript because it processes calls we made to SetupAttachment
+		OriginalActor->ReregisterAllComponents();
+		{
+			// GAllowActorScriptExecutionInEditor must be temporarily true so call to UserConstructionScript isn't skipped
+			FEditorScriptExecutionGuard AllowConstructionScript;
+			OriginalActor->UserConstructionScript();
+		}
+
+#if WITH_EDITOR
+		// Update World Outliner. Usually called by USceneComponent::AttachToComponent.
+		const bool bAttachParentChanged = AttachParentBeforeRestore != AttachParentAfterRestore;
+		if (bAttachParentChanged)
+		{
+			GEngine->BroadcastLevelActorAttached(OriginalActor, AttachParentAfterRestore);
+		}
+#endif
+	}
+
+	static void PreventAttachParentInfiniteRecursion(UActorComponent* Original, const FPropertySelection& PropertySelection)
+	{
+		static const FProperty* AttachParent = USceneComponent::StaticClass()->FindPropertyByName(FName("AttachParent"));
+
+		// Suppose snapshot contains Root > Child and now the hierarchy is Child > Root.
+		// Root's AttachChildren still contains Child after we apply snapshot since that property is transient.
+		// Solution: Detach now, then serialize AttachParent, and OnRegister will automatically call AttachToComponent and update everything.
+		USceneComponent* SceneComponent = Cast<USceneComponent>(Original);
+		if (Original && PropertySelection.IsPropertySelected(nullptr, AttachParent))
+		{
+			SceneComponent->DetachFromComponent(FDetachmentTransformRules::KeepRelativeTransform);
+		}
+	}
+
+	static void UpdateAttachParentAttachChildren(UActorComponent* Original)
+	{
+		// Original->AttachParent had its value serialized but new attach parent's AttachChildren must be updated, too
+		USceneComponent* Component = Cast<USceneComponent>(Original);
+		const bool bNeedsToUpdateParentAttachChildren = Component && Component->GetAttachParent()
+			&& !Component->GetAttachParent()->GetAttachChildren().Contains(Component);
+		if (bNeedsToUpdateParentAttachChildren)
+		{
+			// Hacky way of updating AttachChildren since there is no direct way...
+			Component->UnregisterComponent();
+			Component->SetupAttachment(Component->GetAttachParent());
+			// ReregisterAllComponents will be called after the actor is done restoring which will fix up everything
+		}
+	}
+}
+
+void UE::LevelSnapshots::Private::RestoreIntoExistingWorldActor(AActor* OriginalActor, FActorSnapshotData& ActorData, FWorldSnapshotData& WorldData, FSnapshotDataCache& Cache, UPackage* InLocalisationSnapshotPackage, const FPropertySelectionMap& SelectedProperties)
+{
+	UE_LOG(LogLevelSnapshots, Verbose, TEXT("========== Apply existing %s =========="), *OriginalActor->GetPathName());
+	const bool bWasRecreated = false;
+	const FApplySnapshotToActorScope NotifyExternalListeners({ OriginalActor, SelectedProperties, bWasRecreated });
+	
+	UE::LevelSnapshots::Private::AddAndRemoveComponentsSelectedForRestore(OriginalActor, WorldData, Cache, SelectedProperties, InLocalisationSnapshotPackage);
+
+	auto DeserializeActor = [&ActorData, &WorldData, &Cache, InLocalisationSnapshotPackage, &SelectedProperties](AActor* OriginalActor, AActor* DeserializedActor)
+	{
+		const FRestoreObjectScope FinishRestore = PreActorRestore_EditorWorld(OriginalActor, ActorData.CustomActorSerializationData, WorldData, Cache, SelectedProperties, InLocalisationSnapshotPackage);
+		if (SelectedProperties.GetObjectSelection(OriginalActor).GetPropertySelection() != nullptr)
+		{
+			FApplySnapshotToEditorArchive::ApplyToExistingEditorWorldObject(ActorData.SerializedActorData, WorldData, Cache, OriginalActor, DeserializedActor, SelectedProperties, ActorData.ClassIndex);
+		}
+	};
+	auto DeserializeComponent = [&WorldData, &Cache, &SelectedProperties, InLocalisationSnapshotPackage](FSubobjectSnapshotData& SerializedCompData, UActorComponent* Original, UActorComponent* Deserialized)
+	{
+		const FRestoreObjectScope FinishRestore = PreSubobjectRestore_EditorWorld(Deserialized, Original, WorldData, Cache, SelectedProperties, InLocalisationSnapshotPackage);
+		const FPropertySelection* ComponentSelectedProperties = SelectedProperties.GetObjectSelection(Original).GetPropertySelection();
+		if (ComponentSelectedProperties)
+		{
+			Internal::Restore::PreventAttachParentInfiniteRecursion(Original, *ComponentSelectedProperties);
+			FApplySnapshotToEditorArchive::ApplyToExistingEditorWorldObject(SerializedCompData, WorldData, Cache, Original, Deserialized, SelectedProperties, SerializedCompData.ClassIndex);
+			Internal::Restore::UpdateAttachParentAttachChildren(Original);
+		};
+	};
+	Internal::Restore::DeserializeIntoEditorWorldActor(OriginalActor, ActorData, WorldData, Cache, InLocalisationSnapshotPackage, DeserializeActor, DeserializeComponent);
+}
+
+void UE::LevelSnapshots::Private::RestoreIntoRecreatedEditorWorldActor(AActor* OriginalActor, FActorSnapshotData& ActorData, FWorldSnapshotData& WorldData, FSnapshotDataCache& Cache, UPackage* InLocalisationSnapshotPackage, const FPropertySelectionMap& SelectedProperties)
+{
+	UE_LOG(LogLevelSnapshots, Verbose, TEXT("========== Apply recreated %s =========="), *OriginalActor->GetPathName());
+	const bool bWasRecreated = true;
+	const FApplySnapshotToActorScope NotifyExternalListeners({ OriginalActor, SelectedProperties, bWasRecreated });
+	
+	AllocateMissingComponentsForRecreatedActor(OriginalActor, WorldData, Cache);
+
+	auto DeserializeActor = [&ActorData, &WorldData, &Cache, InLocalisationSnapshotPackage, &SelectedProperties](AActor* OriginalActor, AActor* DeserializedActor)
+	{
+		const FRestoreObjectScope FinishRestore = PreActorRestore_EditorWorld(OriginalActor, ActorData.CustomActorSerializationData, WorldData, Cache, SelectedProperties, InLocalisationSnapshotPackage);
+		FApplySnapshotToEditorArchive::ApplyToEditorWorldObjectRecreatedWithArchetype(ActorData.SerializedActorData, WorldData, Cache, OriginalActor, SelectedProperties);
+	};
+	auto DeserializeComponent = [&WorldData, &Cache, &SelectedProperties, InLocalisationSnapshotPackage](FSubobjectSnapshotData& SerializedCompData, UActorComponent* Original, UActorComponent* Deserialized)
+	{
+		const FRestoreObjectScope FinishRestore = PreSubobjectRestore_EditorWorld(Deserialized, Original, WorldData, Cache, SelectedProperties, InLocalisationSnapshotPackage);
+		// Components are not recreated with any archetype -> they must go through the full serialization process 
+		FApplySnapshotToEditorArchive::ApplyToEditorWorldObjectRecreatedWithoutArchetype(SerializedCompData, WorldData, Cache, Original, Deserialized, SelectedProperties, SerializedCompData.ClassIndex);
+		Internal::Restore::UpdateAttachParentAttachChildren(Original);
+	};
+	Internal::Restore::DeserializeIntoEditorWorldActor(OriginalActor, ActorData, WorldData, Cache, InLocalisationSnapshotPackage, DeserializeActor, DeserializeComponent);
+
+#if WITH_EDITOR
+	// Otherwise actor will show up with internal object name, e.g. actor previously called Cube will be StaticMeshActor1
+	if (ActorData.ActorLabel.IsEmpty()) 
+	{
+		// This might be an old snapshot in which this data was not yet saved
+		OriginalActor->SetActorLabel(OriginalActor->GetName());
+	}
+	else
+	{
+		OriginalActor->SetActorLabel(ActorData.ActorLabel);
+	}
+	
+	// Recreated actors have invalid lightning cache... e.g. recreated point lights will show error image (S_LightError)
+	OriginalActor->InvalidateLightingCacheDetailed(false);
+#endif
+}
diff --git a/Engine/Plugins/VirtualProduction/LevelSnapshots/Source/LevelSnapshots/Private/Data/Util/WorldData/SnapshotObjectUtil.cpp b/Engine/Plugins/VirtualProduction/LevelSnapshots/Source/LevelSnapshots/Private/Data/Util/WorldData/SnapshotObjectUtil.cpp
new file mode 100644
--- /dev/null
+++ b/Engine/Plugins/VirtualProduction/LevelSnapshots/Source/LevelSnapshots/Private/Data/Util/WorldData/SnapshotObjectUtil.cpp
@@ -0,0 +1,420 @@
+﻿// Copyright Epic Games, Inc. All Rights Reserved.
+
+#include "Util/WorldData/SnapshotObjectUtil.h"
+
+#include "Archive/ApplySnapshotToEditorArchive.h"
+#include "Archive/LoadSnapshotObjectArchive.h"
+#include "Archive/TakeWorldObjectSnapshotArchive.h"
+#include "ClassDataUtil.h"
+#include "Data/Util/SnapshotUtil.h"
+#include "Data/Util/WorldData/ActorUtil.h"
+#include "LevelSnapshotsLog.h"
+#include "SnapshotRestorability.h"
+#include "SubobjectSnapshotData.h"
+#include "WorldSnapshotData.h"
+
+#include "Components/ActorComponent.h"
+#include "EngineUtils.h"
+#include "ObjectTrace.h"
+#include "SnapshotDataCache.h"
+#include "GameFramework/Actor.h"
+#include "Templates/NonNullPointer.h"
+
+namespace UE::LevelSnapshots::Private::Internal
+{
+	static UObject* ResolveExternalReference(const FSoftObjectPath& ObjectPath)
+	{
+		UObject* ExternalReference = ObjectPath.ResolveObject();
+		const bool bNeedsToLoadFromDisk = ExternalReference == nullptr;
+		if (bNeedsToLoadFromDisk)
+		{
+			ExternalReference = ObjectPath.TryLoad();
+		}
+		
+		// We're supposed to be dealing with an external (asset) reference, e.g. a UMaterial.
+		ensureMsgf(ExternalReference == nullptr || (!ExternalReference->IsA<AActor>() && !ExternalReference->IsA<UActorComponent>()), TEXT("Something is wrong. We just checked that the reference is not a world object but it is an actor or component."));
+		return ExternalReference;
+	}
+	
+	static UObject* CreateSubobjectForSnapshotWorld(FWorldSnapshotData& WorldData, FSnapshotDataCache& Cache, int32 ObjectPathIndex, const FSoftObjectPath& SnapshotPathToSubobject, const FProcessObjectDependency& ProcessObjectDependency, const FString& LocalisationNamespace)
+	{
+		FSubobjectSnapshotData* SubobjectData = WorldData.Subobjects.Find(ObjectPathIndex);
+		if (!SubobjectData || SubobjectData->bWasSkippedClass)
+		{
+			return nullptr;
+		}
+		if (const FSubobjectSnapshotCache* SubobjectCache = Cache.SubobjectCache.Find(SnapshotPathToSubobject); SubobjectCache && SubobjectCache->SnapshotObject)
+		{
+			return SubobjectCache->SnapshotObject.Get();
+		}
+		
+		const int32 OuterIndex = SubobjectData->OuterIndex;
+		check(OuterIndex != ObjectPathIndex);
+		UObject* SubobjectOuter = ResolveObjectDependencyForSnapshotWorld(WorldData, Cache, OuterIndex, ProcessObjectDependency, LocalisationNamespace);
+		if (!SubobjectOuter)
+		{
+			UE_LOG(LogLevelSnapshots, Warning, TEXT("Failed to create '%s' because its outer could not be created."), *SnapshotPathToSubobject.ToString());
+			return nullptr;
+		}
+	
+		const FName SubobjectName = *ExtractLastSubobjectName(SnapshotPathToSubobject);
+		const bool bSubobjectNameIsTaken = [SubobjectOuter, SubobjectName]()
+		{
+			TArray<UObject*> Subobjects;
+			GetObjectsWithOuter(SubobjectOuter, Subobjects, true);
+			for (UObject* Object : Subobjects)
+			{
+				if (Object->GetFName() == SubobjectName)
+				{
+					return true;
+				}
+			}
+			return false;
+		}();
+		if (bSubobjectNameIsTaken)
+		{
+			UE_LOG(LogLevelSnapshots, Warning, TEXT("Failed to create '%s' because subobject name was already taken."), *SnapshotPathToSubobject.ToString());
+			return nullptr;
+		}
+		
+		const FSubobjectArchetypeFallbackInfo Fallback{ SubobjectOuter, SubobjectName, SubobjectData->GetObjectFlags() };
+		TOptional<TNonNullPtr<UObject>> Archetype = GetSubobjectArchetype(WorldData, SubobjectData->ClassIndex, Cache, Fallback);
+		if (!Archetype)
+		{
+			const FSoftClassPath ClassPath = GetClass(*SubobjectData, WorldData);
+			UE_LOG(LogLevelSnapshots, Warning, TEXT("Failed to create '%s' because class was '%s' not found. Maybe it was removed?"), *SnapshotPathToSubobject.ToString(), *ClassPath.ToString());
+			return nullptr;
+		}
+		
+		const UClass* Class = Archetype->GetClass();
+		UObject* Subobject = NewObject<UObject>(SubobjectOuter, Class, SubobjectName, SubobjectData->GetObjectFlags(), Archetype.GetValue());
+		FSubobjectSnapshotCache& SubobjectCache = Cache.SubobjectCache.FindOrAdd(WorldData.SerializedObjectReferences[ObjectPathIndex]);
+		SubobjectCache.SnapshotObject = Subobject;
+		FLoadSnapshotObjectArchive::ApplyToSnapshotWorldObject(*SubobjectData, WorldData, Cache, Subobject, ProcessObjectDependency, LocalisationNamespace)	;
+		return Subobject;
+	}
+	
+	static UObject* SerializeOrCreateSubobjectForEditorWorld(FWorldSnapshotData& WorldData, FSubobjectSnapshotData& SubobjectData, FSnapshotDataCache& Cache, UObject* ResolvedObject, const FSoftObjectPath& OriginalObjectPath, const FString& LocalisationNamespace, const FPropertySelectionMap& SelectionMap)
+	{
+		using namespace UE::LevelSnapshots::Private;
+
+		if (SubobjectData.bWasSkippedClass)
+		{
+			return !IsValid(ResolvedObject) || ResolvedObject->IsUnreachable() ? nullptr : ResolvedObject;
+		}
+
+		FSubobjectSnapshotCache* SubobjectCache = Cache.SubobjectCache.Find(OriginalObjectPath);
+		if (!ensureMsgf(SubobjectCache, TEXT("Subobject should already have been cached when snapshot world was created. Investigate.")))
+		{
+			return nullptr;
+		}
+		if (SubobjectCache->EditorObject.IsValid())
+		{
+			return SubobjectCache->EditorObject.Get();
+		}
+
+		const FSoftClassPath ClassPath = GetClass(SubobjectData, WorldData);
+		const UClass* ExpectedClass = ClassPath.ResolveClass();
+		if (!ExpectedClass)
+		{
+			UE_LOG(LogLevelSnapshots, Warning, TEXT("Class '%s' could not be resolved."), *ClassPath.ToString());
+			return nullptr;
+		}
+		// Suppose ActorA and ActorB. ActorA is being serialized before ActorB. ActorA has a reference to a component in ActorB.
+		// That component is instanced and was not recreated yet because, unluckily, we're serializing ActorA before ActorB.
+		if (ExpectedClass->IsChildOf(UActorComponent::StaticClass()))
+		{
+			UE_LOG(LogLevelSnapshots, Error, TEXT("Object depends on a component that must be recreated and is owned by another actor. Support for this is not yet implemented."));
+			return nullptr;
+		}
+		
+		// Need to clean rename dead objects otherwise we'll get a name clash when instantiating new object
+		if (ResolvedObject)
+		{
+			const bool bIsReferencingDeadObject = !IsValid(ResolvedObject) || ResolvedObject->IsUnreachable();
+			const bool bIsReferencingDifferentClass = ResolvedObject && !ResolvedObject->GetClass()->IsChildOf(ExpectedClass);
+			if (bIsReferencingDeadObject || bIsReferencingDifferentClass)
+			{
+				const FName NewName = MakeUniqueObjectName(ResolvedObject->GetOuter(), ExpectedClass, *ResolvedObject->GetName().Append(TEXT("_TRASH")));
+				ResolvedObject->Rename(*NewName.ToString());
+				ResolvedObject = nullptr;
+			}
+		}
+
+		if (!ensureMsgf(SubobjectCache->SnapshotObject, TEXT("Subobject pointer should have been set. Investigate.")))
+		{
+			return ResolvedObject;
+		}
+		
+		UObject* SnapshotVersion = SubobjectCache->SnapshotObject.Get();
+		if (ResolvedObject)
+		{
+			// We can't tell whether ResolvedObject is 1: a normal referenced object or 2: a dead object that just has not been collected yet
+			// If it is referenced from somewhere (case 1), we expect it to be in the selection map.
+			// Otherwise we'll treat it as a dead object that we fully restore.
+			if (SelectionMap.IsSubobjectMarkedForReferenceRestorationOnly(OriginalObjectPath))
+			{
+				// No serialization is required. Just return it.
+			}
+			else if (SelectionMap.GetObjectSelection(OriginalObjectPath).GetPropertySelection() != nullptr)
+			{
+				FApplySnapshotToEditorArchive::ApplyToExistingEditorWorldObject(SubobjectData, WorldData, Cache, ResolvedObject, SnapshotVersion, SelectionMap, SubobjectData.ClassIndex);
+			}
+			else
+			{
+				// Reuse existing instance but treat it like it was just freshly allocated
+				FApplySnapshotToEditorArchive::ApplyToEditorWorldObjectRecreatedWithArchetype(SubobjectData, WorldData, Cache, ResolvedObject, SelectionMap);
+			}
+		}
+		else 
+		{
+			UObject* Outer = ResolveObjectDependencyForEditorWorld(WorldData, Cache, SubobjectData.OuterIndex, LocalisationNamespace, SelectionMap);
+			if (!Outer)
+			{
+				UE_LOG(LogLevelSnapshots, Error, TEXT("Failed to create object '%s' because its outer '%s' could not be resolved"), *OriginalObjectPath.ToString(), *WorldData.SerializedObjectReferences[SubobjectData.OuterIndex].ToString());
+				return nullptr;
+			}
+			
+			const FName ObjectName = *ExtractLastSubobjectName(OriginalObjectPath);
+			const FSubobjectArchetypeFallbackInfo Fallback{ Outer, ObjectName, SubobjectData.GetObjectFlags() };
+			const TOptional<TNonNullPtr<UObject>> Archetype = GetSubobjectArchetype(WorldData, SubobjectData.ClassIndex, Cache, Fallback);
+			if (!Archetype)
+			{
+				UE_LOG(LogLevelSnapshots, Error, TEXT("Failed to create object '%s' because its archetype could not be resolved"), *OriginalObjectPath.ToString());
+				return nullptr;	
+			}
+			
+			// Must set RF_Transactional so object creation is transacted
+			ResolvedObject = NewObject<UObject>(Outer, ExpectedClass, *ExtractLastSubobjectName(OriginalObjectPath), SubobjectData.GetObjectFlags() | RF_Transactional, Archetype.GetValue());
+			if (ensureMsgf(ResolvedObject, TEXT("Failed to allocate '%s'"), *OriginalObjectPath.ToString()))
+			{
+				ResolvedObject->SetFlags(SubobjectData.GetObjectFlags());
+				FApplySnapshotToEditorArchive::ApplyToEditorWorldObjectRecreatedWithArchetype(SubobjectData, WorldData, Cache, ResolvedObject, SelectionMap);
+			}
+		}
+		
+		SubobjectCache->EditorObject = ResolvedObject;
+		return ResolvedObject;
+	}
+
+	static void AddSubobjectDependencyInternal(FWorldSnapshotData& WorldData, UObject* ReferenceFromOriginalObject, int32 ObjectIndex)
+	{
+		if (!Restorability::IsSubobjectDesirableForCapture(ReferenceFromOriginalObject))
+		{
+			WorldData.Subobjects.Add(ObjectIndex, FSubobjectSnapshotData::MakeSkippedSubobjectData());
+			return;
+		}
+		
+		const bool bNeedsToSaveSubobject = !WorldData.Subobjects.Contains(ObjectIndex);
+		if (bNeedsToSaveSubobject)
+		{
+			// Avoid infinite recursion
+			WorldData.Subobjects.Add(ObjectIndex);
+			
+			// Owning actor keeps reference to saved subobjects so they can be recreated faster when snapshot is loaded
+			if (!ReferenceFromOriginalObject->IsA<UActorComponent>())
+			{
+				AActor* OwningActor = ReferenceFromOriginalObject->GetTypedOuter<AActor>();
+				check(OwningActor);
+				WorldData.ActorData.FindOrAdd(OwningActor).OwnedSubobjects.Add(ObjectIndex);
+			}
+
+			
+			// Important: first allocate SubobjectData on stack... 
+			FSubobjectSnapshotData SubobjectData; 
+			SubobjectData.ClassIndex = AddClassArchetype(WorldData, ReferenceFromOriginalObject);
+			SubobjectData.OuterIndex = AddObjectDependency(WorldData, ReferenceFromOriginalObject->GetOuter());
+			// ... because serialisation recursively calls this function. FWorldSnapshotData::Subobjects is possibly reallocated
+			FTakeWorldObjectSnapshotArchive::TakeSnapshot(SubobjectData, WorldData, ReferenceFromOriginalObject);
+			
+			WorldData.Subobjects[ObjectIndex] = MoveTemp(SubobjectData); // MoveTemp not profiled
+		}
+	}
+	
+	static void EnsureSubobjectIsSerialized(FWorldSnapshotData& WorldData, FSnapshotDataCache& Cache, UObject* ExistingObject, int32 ObjectPathIndex, const FProcessObjectDependency& ProcessObjectDependency, const FString& LocalisationNamespace)
+	{
+		FSubobjectSnapshotData* SubobjectData = WorldData.Subobjects.Find(ObjectPathIndex);
+		if (!SubobjectData || SubobjectData->bWasSkippedClass)
+		{
+			return;
+		}
+
+		const FSoftClassPath ClassPath = GetClass(*SubobjectData, WorldData);
+		if (ClassPath.ResolveClass() != ExistingObject->GetClass())
+		{
+			UE_LOG(LogLevelSnapshots, Warning, TEXT("Skipping serialisation of subobject because classes are different. Snapshot object '%s' will not contain values saved in snapshot."), *ExistingObject->GetName());
+			return;
+		}
+
+		FSubobjectSnapshotCache& SubobjectCache = Cache.SubobjectCache.FindOrAdd(WorldData.SerializedObjectReferences[ObjectPathIndex]);
+		if (!SubobjectCache.SnapshotObject)
+		{
+			SubobjectCache.SnapshotObject = ExistingObject;
+			FLoadSnapshotObjectArchive::ApplyToSnapshotWorldObject(*SubobjectData, WorldData, Cache, ExistingObject, ProcessObjectDependency, LocalisationNamespace);
+		}
+	}
+	
+	static int32 AddOrFindObjectReference(FWorldSnapshotData& WorldData, const FSoftObjectPath& ObjectPath)
+	{
+		if (const int32* ExistingIndex = WorldData.ReferenceToIndex.Find(ObjectPath))
+		{
+			return *ExistingIndex;
+		}
+		
+		const int32 Index = WorldData.SerializedObjectReferences.Add(ObjectPath);
+		WorldData.ReferenceToIndex.Add(ObjectPath, Index);
+		return Index;
+	}
+}
+
+UObject* UE::LevelSnapshots::Private::ResolveObjectDependencyForSnapshotWorld(FWorldSnapshotData& WorldData, FSnapshotDataCache& Cache, int32 ObjectPathIndex, const FProcessObjectDependency& ProcessObjectDependency, const FString& LocalisationNamespace)
+{
+	if (!ensure(WorldData.SerializedObjectReferences.IsValidIndex(ObjectPathIndex)))
+	{
+		return nullptr;
+	}
+	
+	const FSoftObjectPath& OriginalObjectPath = WorldData.SerializedObjectReferences[ObjectPathIndex];
+	bool bIsPathToActorSubobject;
+	const TOptional<FSoftObjectPath> PathToActor = ExtractActorFromPath(OriginalObjectPath, bIsPathToActorSubobject);
+	
+	const bool bIsExternalAssetReference = !PathToActor.IsSet();
+	if (bIsExternalAssetReference)
+	{
+		return Internal::ResolveExternalReference(OriginalObjectPath);
+	}
+
+	const TOptional<TNonNullPtr<AActor>> Result = GetPreallocated(*PathToActor, WorldData, Cache);
+	if (!Result.IsSet())
+	{
+		return nullptr;
+	}
+
+	// Reference to actor
+	AActor* SnapshotActor = Result.GetValue();
+	const bool bIsPathToActor = !bIsPathToActorSubobject; 
+	if (bIsPathToActor)
+	{
+		return SnapshotActor;
+	}
+
+	// Reference to default subobjects, such as components, or an object we already allocated.
+	const FSoftObjectPath SnapshotPathToSubobject = UE::LevelSnapshots::Private::SetActorInPath(SnapshotActor, OriginalObjectPath);
+	if (UObject* ExistingSubobject = SnapshotPathToSubobject.ResolveObject())
+	{
+		Internal::EnsureSubobjectIsSerialized(WorldData, Cache, ExistingSubobject, ObjectPathIndex, ProcessObjectDependency, LocalisationNamespace);
+		return ExistingSubobject;
+	}
+
+	// The subobject is instanced: recreate it recursively
+	return Internal::CreateSubobjectForSnapshotWorld(WorldData, Cache, ObjectPathIndex, SnapshotPathToSubobject, ProcessObjectDependency, LocalisationNamespace);
+}
+
+UObject* UE::LevelSnapshots::Private::ResolveObjectDependencyForEditorWorld(FWorldSnapshotData& WorldData, FSnapshotDataCache& Cache, int32 ObjectPathIndex, const FString& LocalisationNamespace, const FPropertySelectionMap& SelectionMap)
+{
+	if (!ensure(WorldData.SerializedObjectReferences.IsValidIndex(ObjectPathIndex)))
+	{
+		return nullptr;
+	}
+	
+	const FSoftObjectPath& OriginalObjectPath = WorldData.SerializedObjectReferences[ObjectPathIndex];
+	bool bIsPathToActorSubobject;
+	const TOptional<TNonNullPtr<FActorSnapshotData>> SnapshotData = UE::LevelSnapshots::Private::FindSavedActorDataUsingObjectPath(WorldData.ActorData, OriginalObjectPath, bIsPathToActorSubobject);
+	
+	// Dealing with an external (asset) reference, e.g. a UMaterial?
+	const bool bIsExternalAssetReference = !SnapshotData.IsSet();
+	if (bIsExternalAssetReference)
+	{
+		return Internal::ResolveExternalReference(OriginalObjectPath);
+	}
+
+	// Can return immediately for actors and components: their serialization is handled by FActorSnapshotData.
+	UObject* ResolvedObject = OriginalObjectPath.ResolveObject();
+	if (ResolvedObject && (ResolvedObject->IsA<AActor>() || ResolvedObject->IsA<UActorComponent>()))
+	{
+		// Might resolve to a dead object which still exists in memory
+		return IsValid(ResolvedObject) ? ResolvedObject : nullptr;
+	}
+	
+	if (FSubobjectSnapshotData* const SubobjectData = WorldData.Subobjects.Find(ObjectPathIndex))
+	{
+		return Internal::SerializeOrCreateSubobjectForEditorWorld(WorldData, *SubobjectData, Cache, ResolvedObject, OriginalObjectPath, LocalisationNamespace, SelectionMap);
+	}
+	
+	// This should not happen after release of 4.27. This is error may happen if certain snapshots in an older data format are applied because the data wasn't yet captured.
+	UE_LOG(LogLevelSnapshots, Warning, TEXT("Trying to resolve a subobject but found no associated object data."));
+	return ResolvedObject;
+}
+
+UObject* UE::LevelSnapshots::Private::ResolveObjectDependencyForClassDefaultObject(FWorldSnapshotData& WorldData, int32 ObjectPathIndex)
+{
+	if (ensure(WorldData.SerializedObjectReferences.IsValidIndex(ObjectPathIndex)))
+	{
+		const FSoftObjectPath& OriginalObjectPath = WorldData.SerializedObjectReferences[ObjectPathIndex];
+		return Internal::ResolveExternalReference(OriginalObjectPath);
+	}
+	return nullptr;
+}
+
+int32 UE::LevelSnapshots::Private::AddObjectDependency(FWorldSnapshotData& WorldData, UObject* ReferenceFromOriginalObject, bool bCheckWhetherSubobject)
+{
+	// Even if FSnapshotRestorability::IsSubobjectDesirableForCapture later returns false for this object, we want to track it
+	const int32 Result = Internal::AddOrFindObjectReference(WorldData, ReferenceFromOriginalObject);
+
+	if (bCheckWhetherSubobject && ReferenceFromOriginalObject && ReferenceFromOriginalObject->GetTypedOuter<AActor>())
+	{
+		Internal::AddSubobjectDependencyInternal(WorldData, ReferenceFromOriginalObject, Result);
+	}
+	
+	return Result;
+}
+
+int32 UE::LevelSnapshots::Private::AddCustomSubobjectDependency(FWorldSnapshotData& WorldData, UObject* ReferenceFromOriginalObject)
+{
+	if (!ensure(ReferenceFromOriginalObject))
+	{
+		return INDEX_NONE;
+	}
+
+	const int32 SubobjectIndex = Internal::AddOrFindObjectReference(WorldData, ReferenceFromOriginalObject);
+	if (WorldData.CustomSubobjectSerializationData.Find(SubobjectIndex) == nullptr)
+	{
+		WorldData.CustomSubobjectSerializationData.Add(SubobjectIndex);
+	}
+	else
+	{
+		UE_LOG(LogLevelSnapshots, Warning, TEXT("Object %s was already added as dependency. Investigate"), *ReferenceFromOriginalObject->GetName());
+		UE_DEBUG_BREAK();
+	}
+		
+	return SubobjectIndex;
+}
+
+FCustomSerializationData* UE::LevelSnapshots::Private::FindCustomSubobjectData(FWorldSnapshotData& WorldData,const FSoftObjectPath& ReferenceFromOriginalObject)
+{
+	const int32 SubobjectIndex = WorldData.SerializedObjectReferences.Find(ReferenceFromOriginalObject);
+	return WorldData.CustomSubobjectSerializationData.Find(SubobjectIndex);
+}
+
+const FCustomSerializationData* UE::LevelSnapshots::Private::FindCustomActorOrSubobjectData(const FWorldSnapshotData& WorldData,UObject* OriginalObject)
+{
+	// Is it an actor?
+	if (const FActorSnapshotData* SavedActorData = WorldData.ActorData.Find(OriginalObject))
+	{
+		return &SavedActorData->CustomActorSerializationData;
+	}
+	if (Cast<AActor>(OriginalObject))
+	{
+		// Return immediately to avoid searching the entire array below.
+		return nullptr;
+	}
+
+	// If not an actor, it is a subobject
+	const int32 ObjectReferenceIndex = WorldData.SerializedObjectReferences.Find(OriginalObject);
+	if (ObjectReferenceIndex != INDEX_NONE)
+	{
+		return WorldData.CustomSubobjectSerializationData.Find(ObjectReferenceIndex);
+	}
+
+	return nullptr;
+}
diff --git a/Engine/Plugins/VirtualProduction/LevelSnapshots/Source/LevelSnapshots/Private/Data/Util/WorldData/SnapshotObjectUtil.h b/Engine/Plugins/VirtualProduction/LevelSnapshots/Source/LevelSnapshots/Private/Data/Util/WorldData/SnapshotObjectUtil.h
new file mode 100644
--- /dev/null
+++ b/Engine/Plugins/VirtualProduction/LevelSnapshots/Source/LevelSnapshots/Private/Data/Util/WorldData/SnapshotObjectUtil.h
@@ -0,0 +1,58 @@
+﻿// Copyright Epic Games, Inc. All Rights Reserved.
+
+#pragma once
+
+#include "CoreMinimal.h"
+#include "Misc/ObjectDependencyCallback.h"
+
+class AActor;
+class UObject;
+
+struct FCustomSerializationData;
+struct FPropertySelectionMap;
+struct FSnapshotDataCache;
+struct FWorldSnapshotData;
+
+namespace UE::LevelSnapshots::Private
+{
+	/* Resolves an object dependency for use in the snapshot world. If the object is a subobject, it gets fully serialized. */
+	UObject* ResolveObjectDependencyForSnapshotWorld(FWorldSnapshotData& WorldData, FSnapshotDataCache& Cache, int32 ObjectPathIndex, const FProcessObjectDependency& ProcessObjectDependency, const FString& LocalisationNamespace); 
+
+	/* Resolves an object dependency for use in the editor world. If the object is a subobject, it is serialized.
+	 * Steps for serializing subobject:
+	 *  - If an equivalent object with the saved name and class exists, use that and serialize the properties in SelectionMap into it.
+	 *  - If object does not exist, allocate a new object and serialize all properties into it
+	 *  - If object exists but has different class, skip. Return null.
+	 */
+	UObject* ResolveObjectDependencyForEditorWorld(FWorldSnapshotData& WorldData, FSnapshotDataCache& Cache, int32 ObjectPathIndex, const FString& LocalisationNamespace, const FPropertySelectionMap& SelectionMap);
+
+	/* Resolves an object depedency when restoring a class default object. Simply resolves without further checks. */
+	UObject* ResolveObjectDependencyForClassDefaultObject(FWorldSnapshotData& WorldData, int32 ObjectPathIndex);
+
+
+	
+	/**
+	 * Saves the object's path. Intended for external objects, e.g. to UMaterial in the content browser.
+	 *
+	 * @param ReferenceFromOriginalObject The reference as it was found in a property in the editor world
+	 * @param bCheckWhetherSubobject If true, this will check the the object is a subobject of an actor and call AddSubobjectDependency.
+	 */
+	int32 AddObjectDependency(FWorldSnapshotData& WorldData, UObject* ReferenceFromOriginalObject, bool bCheckWhetherSubobject = true);
+
+
+
+	
+	/**
+	 * Adds a subobject to SerializedObjectReferences and CustomSubobjectSerializationData.
+	 * @return A valid index in SerializedObjectReferences and the corresponding subobject data.
+	 */
+	int32 AddCustomSubobjectDependency(FWorldSnapshotData& WorldData, UObject* ReferenceFromOriginalObject);
+
+	/** Tries to find custom serialisation data. Use this when you already know that ReferenceFromOriginalObject is not an actor. */
+	FCustomSerializationData* FindCustomSubobjectData(FWorldSnapshotData& WorldData, const FSoftObjectPath& ReferenceFromOriginalObject);
+	
+	/** Tries to find custom serialisation data. Use when OriginalObject might be an actor . */
+	const FCustomSerializationData* FindCustomActorOrSubobjectData(const FWorldSnapshotData& WorldData, UObject* OriginalObject);
+}
+
+
diff --git a/Engine/Plugins/VirtualProduction/LevelSnapshots/Source/LevelSnapshots/Private/Data/Util/WorldData/WorldDataUtil.h b/Engine/Plugins/VirtualProduction/LevelSnapshots/Source/LevelSnapshots/Private/Data/Util/WorldData/WorldDataUtil.h
new file mode 100644
--- /dev/null
+++ b/Engine/Plugins/VirtualProduction/LevelSnapshots/Source/LevelSnapshots/Private/Data/Util/WorldData/WorldDataUtil.h
@@ -0,0 +1,22 @@
+﻿// Copyright Epic Games, Inc. All Rights Reserved.
+
+#pragma once
+
+#include "CoreMinimal.h"
+
+class AActor;
+class UActorComponent;
+class UWorld;
+
+struct FPropertySelectionMap;
+struct FSnapshotDataCache;
+struct FWorldSnapshotData;
+
+namespace UE::LevelSnapshots::Private
+{
+	/** Captures the state of the world */
+	FWorldSnapshotData SnapshotWorld(UWorld* World);
+	
+	/* Applies the saved properties to WorldActor */
+	void ApplyToWorld(FWorldSnapshotData& WorldData, FSnapshotDataCache& Cache, UWorld* WorldToApplyTo, UPackage* LocalisationSnapshotPackage, const FPropertySelectionMap& PropertiesToSerialize);
+}
diff --git a/Engine/Plugins/VirtualProduction/LevelSnapshots/Source/LevelSnapshots/Private/Data/WorldSnapshotData.cpp b/Engine/Plugins/VirtualProduction/LevelSnapshots/Source/LevelSnapshots/Private/Data/WorldSnapshotData.cpp
--- a/Engine/Plugins/VirtualProduction/LevelSnapshots/Source/LevelSnapshots/Private/Data/WorldSnapshotData.cpp
+++ b/Engine/Plugins/VirtualProduction/LevelSnapshots/Source/LevelSnapshots/Private/Data/WorldSnapshotData.cpp
@@ -3,9 +3,12 @@
 #include "Data/WorldSnapshotData.h"
 
 #include "LevelSnapshotsModule.h"
+#include "SnapshotCustomVersion.h"
 #include "Util/Property/PropertyIterator.h"
 
 #include "Serialization/BufferArchive.h"
+#include "Util/WorldData/ClassDataUtil.h"
+#include "Util/WorldData/CompressionUtil.h"
 
 #define LOCTEXT_NAMESPACE "LevelSnapshotsEditor"
 
@@ -29,15 +32,31 @@ namespace UE::LevelSnapshots::Private::Internal
 		for (auto ActorIt = WorldData.ActorData.CreateIterator(); ActorIt; ++ActorIt)
 		{
 			Ar << ActorIt->Key;
-			Ar << ActorIt->Value.ActorClass;
+			FSoftClassPath ClassPath = GetClass(ActorIt->Value, WorldData);
+			Ar << ClassPath;
 		}
 	}
 
 	static void CollectClassDefaultReferences(FWorldSnapshotData& WorldData, FArchive& Ar)
 	{
-		for (auto ClassDefaultIt = WorldData.ClassDefaults.CreateIterator(); ClassDefaultIt; ++ClassDefaultIt)
+		for (auto ClassDefaultIt = WorldData.ClassData.CreateIterator(); ClassDefaultIt; ++ClassDefaultIt)
 		{
-			Ar << ClassDefaultIt->Key;
+			Ar << ClassDefaultIt->ClassPath;
+		}
+	}
+
+	static void CollectEnums(FArchive& Ar)
+	{
+		TArray<UEnum*> Enums = { StaticEnum<ESnapshotClassFlags>() };
+		for (UEnum* Enum : Enums)
+		{
+			FName ClassFlagsName = Enum->GetFName();
+			Ar << ClassFlagsName;
+			for (int32 i = 0; i < Enum->NumEnums(); ++i)
+			{
+				FName ValueName = Enum->GetNameByIndex(i);
+				Ar << ValueName;
+			}
 		}
 	}
 	
@@ -51,6 +70,7 @@ namespace UE::LevelSnapshots::Private::Internal
 		// Names
 		Ar << WorldData.SerializedNames;
 		Ar << WorldData.SnapshotVersionInfo.CustomVersions;
+		CollectEnums(Ar);
 
 		// Serialize hardcoded property names
 		for (int32 i = 0; i < static_cast<int32>(EName::MaxHardcodedNameIndex); ++i)
@@ -75,23 +95,126 @@ namespace UE::LevelSnapshots::Private::Internal
 		FPropertyIterator(FActorSnapshotData::StaticStruct(), ProcessProperty, ProcessStruct);
 		FPropertyIterator(FSubobjectSnapshotData::StaticStruct(), ProcessProperty, ProcessStruct);
 		FPropertyIterator(FCustomSerializationData::StaticStruct(), ProcessProperty, ProcessStruct);
+		FPropertyIterator(FClassSnapshotData::StaticStruct(), ProcessProperty, ProcessStruct);
+	}
+
+	PRAGMA_DISABLE_DEPRECATION_WARNINGS
+	static void MigrateClassDefaultData(FWorldSnapshotData& WorldData)
+	{
+		TMap<FSoftClassPath, int32> PreventDoubleAddition;
+		
+		// Actor data
+		for (auto ActorIt = WorldData.ActorData.CreateIterator(); ActorIt; ++ActorIt)
+		{
+			FActorSnapshotData& ActorData = ActorIt->Value;
+			const FSoftClassPath ClassPath = ActorData.ActorClass_DEPRECATED;
+			ActorData.ActorClass_DEPRECATED.Reset();
+			
+			if (const int32* Existing = PreventDoubleAddition.Find(ClassPath))
+			{
+				ActorData.ClassIndex = *Existing;
+				continue;
+			}
+			
+			const FClassDefaultObjectSnapshotData* ClassDefaultData = WorldData.ClassDefaults_DEPRECATED.Find(ClassPath);
+			if (ClassDefaultData)
+			{
+				FClassSnapshotData ClassData;
+				ClassData.SerializedData = ClassDefaultData->SerializedData;
+				ClassData.ObjectFlags = ClassDefaultData->ObjectFlags;
+				ClassData.ClassPath = ClassPath;
+
+				const int32 ClassIndex = WorldData.ClassData.Add(ClassData);
+				ActorData.ClassIndex = ClassIndex;
+				PreventDoubleAddition.Add(ClassPath, ClassIndex);
+			}
+			else
+			{
+				checkNoEntry();
+			}
+		}
+
+		// Subobject data
+		for (auto SubobjectIt = WorldData.Subobjects.CreateIterator(); SubobjectIt; ++SubobjectIt)
+		{
+			FSubobjectSnapshotData& SubobjectData = SubobjectIt->Value;
+			if (SubobjectData.bWasSkippedClass)
+			{
+				continue;
+			}
+			
+			const FSoftClassPath ClassPath = SubobjectData.Class_DEPRECATED;
+			SubobjectData.Class_DEPRECATED.Reset();
+			
+			if (const int32* Existing = PreventDoubleAddition.Find(ClassPath))
+			{
+				SubobjectData.ClassIndex = *Existing;
+				continue;
+			}
+			
+			const FClassDefaultObjectSnapshotData* ClassDefaultData = WorldData.ClassDefaults_DEPRECATED.Find(ClassPath);
+			if (ClassDefaultData)
+			{
+				FClassSnapshotData ClassData;
+				ClassData.SerializedData = ClassDefaultData->SerializedData;
+				ClassData.ObjectFlags = ClassDefaultData->ObjectFlags;
+				ClassData.ClassPath = ClassPath;
+			
+				const int32 ClassIndex = WorldData.ClassData.Add(ClassData);
+				SubobjectData.ClassIndex = ClassIndex;
+				PreventDoubleAddition.Add(ClassPath, ClassIndex);
+			}
+			else
+			{
+				checkNoEntry();
+			}
+		}
+
+		WorldData.ClassDefaults_DEPRECATED.Reset();
 	}
+	PRAGMA_ENABLE_DEPRECATION_WARNINGS
 }
 
 bool FWorldSnapshotData::Serialize(FArchive& Ar)
 {
+	using namespace UE::LevelSnapshots::Private;
+	Ar.UsingCustomVersion(FSnapshotCustomVersion::GUID);
+	
 	// When this struct is saved, the save algorithm collects references. It's faster if we just give it the info directly.
 	if (Ar.IsObjectReferenceCollector())
 	{
 		UE::LevelSnapshots::Private::Internal::CollectReferencesAndNames(*this, Ar);
 		return true;
 	}
+
+	if (Ar.IsSaving())
+	{
+		UE::LevelSnapshots::Compress(Ar, *this);
+		return true;
+	}
+
+	if (Ar.IsLoading())
+	{
+		if (Ar.CustomVer(FSnapshotCustomVersion::GUID) >= FSnapshotCustomVersion::OoddleCompression)
+		{
+			UE::LevelSnapshots::Decompress(Ar, *this);
+			return true;
+		}
+		else
+		{
+			// Use tagged property serialization
+			return false;
+		}
+	}
 	
 	return false;
 }
 
 void FWorldSnapshotData::PostSerialize(const FArchive& Ar)
 {
+	using namespace UE::LevelSnapshots::Private;
+	using namespace UE::LevelSnapshots::Private::Internal;
+	
 	if (Ar.IsLoading() && !SnapshotVersionInfo.IsInitialized())
 	{
 		// Assets saved before we added version tracking need to receive versioning info of 4.27.
@@ -99,6 +222,12 @@ void FWorldSnapshotData::PostSerialize(const FArchive& Ar)
 		const bool bWithoutSnapshotVersion = true;
 		SnapshotVersionInfo.Initialize(bWithoutSnapshotVersion);
 	}
+	
+	if (Ar.IsLoading()
+		&& Ar.CustomVer(FSnapshotCustomVersion::GUID) < FSnapshotCustomVersion::ClassArchetypeRefactor)
+	{
+		MigrateClassDefaultData(*this);
+	}
 }
 
 #undef LOCTEXT_NAMESPACE
\ No newline at end of file
diff --git a/Engine/Plugins/VirtualProduction/LevelSnapshots/Source/LevelSnapshots/Private/LevelSnapshotsFunctionLibrary.cpp b/Engine/Plugins/VirtualProduction/LevelSnapshots/Source/LevelSnapshots/Private/LevelSnapshotsFunctionLibrary.cpp
--- a/Engine/Plugins/VirtualProduction/LevelSnapshots/Source/LevelSnapshots/Private/LevelSnapshotsFunctionLibrary.cpp
+++ b/Engine/Plugins/VirtualProduction/LevelSnapshots/Source/LevelSnapshots/Private/LevelSnapshotsFunctionLibrary.cpp
@@ -3,12 +3,11 @@
 #include "LevelSnapshotsFunctionLibrary.h"
 
 #include "Selection/ApplySnapshotFilter.h"
-#include "ConstantFilter.h"
 #include "LevelSnapshot.h"
 #include "LevelSnapshotFilters.h"
-#include "LevelSnapshotsLog.h"
 
 #include "EngineUtils.h"
+#include "LevelSnapshotsFilteringLibrary.h"
 #include "CustomSerialization/CustomObjectSerializationWrapper.h"
 
 ULevelSnapshot* ULevelSnapshotsFunctionLibrary::TakeLevelSnapshot(const UObject* WorldContextObject, const FName NewSnapshotName, const FString Description)
@@ -36,91 +35,13 @@ ULevelSnapshot* ULevelSnapshotsFunctionLibrary::TakeLevelSnapshot_Internal(const
 	return NewSnapshot;
 }
 
-namespace
-{
-	struct FPropertyMapBuilder
-	{
-		FPropertyMapBuilder(UWorld* TargetWorld, ULevelSnapshot* Snapshot, ULevelSnapshotFilter* FilterToApply)
-			:
-			Snapshot(Snapshot)
-		{
-			Snapshot->DiffWorld(
-				TargetWorld,
-				ULevelSnapshot::FActorPathConsumer::CreateRaw(this, &FPropertyMapBuilder::HandleActorExistsInWorldAndSnapshot, FilterToApply),
-				ULevelSnapshot::FActorPathConsumer::CreateRaw(this, &FPropertyMapBuilder::HandleActorWasRemovedFromWorld, FilterToApply),
-				ULevelSnapshot::FActorConsumer::CreateRaw(this, &FPropertyMapBuilder::HandleActorWasAddedToWorld, FilterToApply)
-			);
-		}
-
-		const FPropertySelectionMap& GetSelectionMap() const
-		{
-			return SelectionMap;
-		}
-
-	private:
-		
-		FPropertySelectionMap SelectionMap;
-		ULevelSnapshot* Snapshot;
-		
-		void HandleActorExistsInWorldAndSnapshot(const FSoftObjectPath& OriginalActorPath, ULevelSnapshotFilter* FilterToApply)
-		{
-			UObject* ResolvedWorldActor = OriginalActorPath.ResolveObject();
-            if (!ResolvedWorldActor)
-            {
-            	UE_LOG(LogLevelSnapshots, Warning, TEXT("Failed to resolve actor %s. Was it deleted from the world?"), *OriginalActorPath.ToString());
-            	return;
-            }
-        
-            AActor* WorldActor = Cast<AActor>(ResolvedWorldActor);
-            if (ensureAlwaysMsgf(WorldActor, TEXT("A path that was previously associated with an actor no longer refers to an actor. Something is wrong."))
-				&& Snapshot->HasChangedSinceSnapshotWasTaken(WorldActor))
-            {
-            	const TOptional<TNonNullPtr<AActor>> DeserializedSnapshotActor = Snapshot->GetDeserializedActor(OriginalActorPath);
-            	if (!ensureMsgf(DeserializedSnapshotActor.Get(nullptr), TEXT("Failed to get TMap value for key %s. Is the snapshot corrupted?"), *OriginalActorPath.ToString()))
-            	{
-            		// Engine issue. Take snapshot. Rename actor. Update references. Value is updated correctly in TMap but look ups no longer work.
-            		UE_LOG(LogLevelSnapshots, Error, TEXT("Failed to lookup actor %s OriginalActorPath. The snapshot is corrupted."));
-            		return;
-            	}
-            	ULevelSnapshotsFunctionLibrary::ApplyFilterToFindSelectedProperties(Snapshot, SelectionMap, WorldActor, DeserializedSnapshotActor.GetValue(), FilterToApply);
-            }
-		}
-		
-		void HandleActorWasRemovedFromWorld(const FSoftObjectPath& OriginalActorPath, ULevelSnapshotFilter* FilterToApply)
-		{
-			const EFilterResult::Type FilterResult = FilterToApply->IsDeletedActorValid(
-				FIsDeletedActorValidParams(
-					OriginalActorPath,
-					[this](const FSoftObjectPath& ObjectPath)
-					{
-						return Snapshot->GetDeserializedActor(ObjectPath).Get(nullptr);
-					}
-				)
-			);
-			if (EFilterResult::CanInclude(FilterResult))
-			{
-				SelectionMap.AddDeletedActorToRespawn(OriginalActorPath);
-			}
-		}
-		
-		void HandleActorWasAddedToWorld(AActor* WorldActor, ULevelSnapshotFilter* FilterToApply)
-		{
-			const EFilterResult::Type FilterResult = FilterToApply->IsAddedActorValid(FIsAddedActorValidParams(WorldActor)); 
-			if (EFilterResult::CanInclude(FilterResult))
-			{
-				SelectionMap.AddNewActorToDespawn(WorldActor);
-			}
-		}
-	};
-}
-
 void ULevelSnapshotsFunctionLibrary::ApplySnapshotToWorld(const UObject* WorldContextObject, ULevelSnapshot* Snapshot, ULevelSnapshotFilter* OptionalFilter)
 {
 	UWorld* TargetWorld = WorldContextObject ? WorldContextObject->GetWorld() : nullptr;
 	if (ensure(TargetWorld && Snapshot))
 	{
-		const FPropertyMapBuilder Helper(TargetWorld, Snapshot, OptionalFilter ? OptionalFilter : GetMutableDefault<UConstantFilter>());
-		Snapshot->ApplySnapshotToWorld(TargetWorld, Helper.GetSelectionMap());
+		const FPropertySelectionMap SelectionMap = ULevelSnapshotsFilteringLibrary::DiffAndFilterSnapshot(TargetWorld, Snapshot, OptionalFilter);
+		Snapshot->ApplySnapshotToWorld(TargetWorld, SelectionMap);
 	}
 }
 
@@ -133,15 +54,7 @@ void ULevelSnapshotsFunctionLibrary::ApplyFilterToFindSelectedProperties(
 	bool bAllowUnchangedProperties,
     bool bAllowNonEditableProperties)
 {
-	if (Filter == nullptr)
-	{
-		Filter = GetMutableDefault<UConstantFilter>();
-	}
-	
-	UE::LevelSnapshots::Private::FApplySnapshotFilter::Make(Snapshot, DeserializedSnapshotActor, WorldActor, Filter)
-		.AllowUnchangedProperties(bAllowUnchangedProperties)
-		.AllowNonEditableProperties(bAllowNonEditableProperties)
-		.ApplyFilterToFindSelectedProperties(MapToAddTo);
+	ULevelSnapshotsFilteringLibrary::ApplyFilterToFindSelectedProperties(Snapshot, MapToAddTo, WorldActor, DeserializedSnapshotActor, Filter);
 }
 
 void ULevelSnapshotsFunctionLibrary::ForEachMatchingCustomSubobjectPair(
diff --git a/Engine/Plugins/VirtualProduction/LevelSnapshots/Source/LevelSnapshots/Private/LevelSnapshotsModule.cpp b/Engine/Plugins/VirtualProduction/LevelSnapshots/Source/LevelSnapshots/Private/LevelSnapshotsModule.cpp
--- a/Engine/Plugins/VirtualProduction/LevelSnapshots/Source/LevelSnapshots/Private/LevelSnapshotsModule.cpp
+++ b/Engine/Plugins/VirtualProduction/LevelSnapshots/Source/LevelSnapshots/Private/LevelSnapshotsModule.cpp
@@ -6,94 +6,20 @@
 #include "LevelSnapshotsSettings.h"
 #include "LevelSnapshotsLog.h"
 #include "Params/PropertyComparisonParams.h"
-#include "Restorability/CollisionRestoration.h"
-#include "Restorability/GridPlacementRestoration.h"
 
 #include "Algo/AllOf.h"
 #include "Components/ActorComponent.h"
 #include "EngineUtils.h"
-#include "Algo/Transform.h"
-#include "Engine/Brush.h"
 #include "GameFramework/Actor.h"
-#include "Materials/MaterialInstance.h"
 #include "Modules/ModuleManager.h"
+#include "Restorability/EngineTypesRestorationFence.h"
 #if WITH_EDITOR
 #include "ISettingsModule.h"
 #endif
 
 namespace UE::LevelSnapshots::Private::Internal
 {
-	static void AddSoftObjectPathSupport(UE::LevelSnapshots::Private::FLevelSnapshotsModule& Module)
-	{
-		// FSnapshotRestorability::IsRestorableProperty requires properties to have the CPF_Edit specifier
-		// FSoftObjectPath does not have this so we need to explicitly allow its properties
-
-		UStruct* SoftObjectClassPath = FindObject<UStruct>(nullptr, TEXT("/Script/CoreUObject.SoftObjectPath"));
-		if (!ensureMsgf(SoftObjectClassPath, TEXT("Investigate why this class could not be found")))
-		{
-			return;
-		}
-
-		TSet<const FProperty*> SoftObjectPathProperties;
-		Algo::Transform(TFieldRange<const FProperty>(SoftObjectClassPath), SoftObjectPathProperties, [](const FProperty* Prop) { return Prop;} );
-		Module.AddExplicitilySupportedProperties(SoftObjectPathProperties);
-	}
-
-	static void AddAttachParentSupport(UE::LevelSnapshots::Private::FLevelSnapshotsModule& Module)
-	{
-		// These properties are not visible by default because they're not CPF_Edit
-		const FProperty* AttachParent = USceneComponent::StaticClass()->FindPropertyByName(FName("AttachParent"));
-		const FProperty* AttachSocketName = USceneComponent::StaticClass()->FindPropertyByName(FName("AttachSocketName"));
-		// RootComponent is usually set automatically but sometimes not... for example spawning AActor with instanced components only
-		const FProperty* RootComponent = AActor::StaticClass()->FindPropertyByName(FName("RootComponent"));
-		if (ensure(AttachParent && AttachSocketName))
-		{
-			Module.AddExplicitilySupportedProperties({ AttachParent, AttachSocketName, RootComponent });
-		}
-	}
-
-	static void DisableIrrelevantBrushSubobjects(UE::LevelSnapshots::Private::FLevelSnapshotsModule& Module)
-	{
-#if WITH_EDITORONLY_DATA
-		// ABrush::BrushBuilder is CPF_Edit but no user ever cares about it. We don't want it to make volumes to show up as changed.
-		const FProperty* BrushBuilder = ABrush::StaticClass()->FindPropertyByName(GET_MEMBER_NAME_CHECKED(ABrush, BrushBuilder));
-		if (ensure(BrushBuilder))
-		{
-			Module.AddExplicitlyUnsupportedProperties({ BrushBuilder });
-		}
-#endif
-	}
-
-	static void DisableIrrelevantWorldSettings(UE::LevelSnapshots::Private::FLevelSnapshotsModule& Module)
-	{
-		// AWorldSettings::NavigationSystemConfig is CPF_Edit but no user ever cares about it.
-		const FProperty* NavigationSystemConfig = AWorldSettings::StaticClass()->FindPropertyByName(FName("NavigationSystemConfig"));
-		if (ensure(NavigationSystemConfig))
-		{
-			Module.AddExplicitlyUnsupportedProperties({ NavigationSystemConfig });
-		}
-	}
-
-	static void DisableIrrelevantMaterialInstanceProperties(UE::LevelSnapshots::Private::FLevelSnapshotsModule& Module)
-	{
-		// This property causes diffs sometimes for unexplained reasons when creating in construction script... does not seem to be important
-		const FProperty* BasePropertyOverrides = UMaterialInstance::StaticClass()->FindPropertyByName(GET_MEMBER_NAME_CHECKED(UMaterialInstance, BasePropertyOverrides));
-		if (ensure(BasePropertyOverrides))
-		{
-			Module.AddExplicitlyUnsupportedProperties({ BasePropertyOverrides });
-		}
-	}
-
-	static void DisableIrrelevantActorProperties(UE::LevelSnapshots::Private::FLevelSnapshotsModule& Module)
-	{
-#if WITH_EDITORONLY_DATA
-		const FProperty* ActorGuid = AActor::StaticClass()->FindPropertyByName(FName("ActorGuid"));
-		if (ensure(ActorGuid))
-		{
-			Module.AddExplicitlyUnsupportedProperties({ ActorGuid });
-		}
-#endif
-	}
+	
 }
 
 UE::LevelSnapshots::Private::FLevelSnapshotsModule& UE::LevelSnapshots::Private::FLevelSnapshotsModule::GetInternalModuleInstance()
@@ -117,18 +43,9 @@ void UE::LevelSnapshots::Private::FLevelSnapshotsModule::StartupModule()
 		})
 	);
 	RegisterRestorabilityOverrider(ClassSkipper);
-
-	// Enable / disable troublesome properties
-	UE::LevelSnapshots::Private::Internal::AddSoftObjectPathSupport(*this);
-	UE::LevelSnapshots::Private::Internal::AddAttachParentSupport(*this);
-	UE::LevelSnapshots::Private::Internal::DisableIrrelevantBrushSubobjects(*this);
-	UE::LevelSnapshots::Private::Internal::DisableIrrelevantWorldSettings(*this);
-	UE::LevelSnapshots::Private::Internal::DisableIrrelevantMaterialInstanceProperties(*this);
-	UE::LevelSnapshots::Private::Internal::DisableIrrelevantActorProperties(*this);
-
-	// Interact with special engine features
-	FCollisionRestoration::Register(*this);
-	GridPlacementRestoration::Register(*this);
+	
+	// Add support for engine features that require specialized case handling
+	EngineTypesRestorationFence::RegisterSpecialEngineTypeSupport(*this);
 
 #if WITH_EDITOR
 	ISettingsModule& SettingsModule = FModuleManager::LoadModuleChecked<ISettingsModule>("Settings");
diff --git a/Engine/Plugins/VirtualProduction/LevelSnapshots/Source/LevelSnapshots/Private/Restorability/GridPlacementRestoration.h b/Engine/Plugins/VirtualProduction/LevelSnapshots/Source/LevelSnapshots/Private/Restorability/GridPlacementRestoration.h
--- a/Engine/Plugins/VirtualProduction/LevelSnapshots/Source/LevelSnapshots/Private/Restorability/GridPlacementRestoration.h
+++ b/Engine/Plugins/VirtualProduction/LevelSnapshots/Source/LevelSnapshots/Private/Restorability/GridPlacementRestoration.h
@@ -5,8 +5,6 @@
 #include "CoreMinimal.h"
 #include "LevelSnapshotsModule.h"
 
-class UStaticMeshComponent;
-
 namespace UE::LevelSnapshots::Private::GridPlacementRestoration
 {
 	void Register(FLevelSnapshotsModule& Module);
diff --git a/Engine/Plugins/VirtualProduction/LevelSnapshots/Source/LevelSnapshots/Private/Restorability/SnapshotRestorability.cpp b/Engine/Plugins/VirtualProduction/LevelSnapshots/Source/LevelSnapshots/Private/Restorability/SnapshotRestorability.cpp
--- a/Engine/Plugins/VirtualProduction/LevelSnapshots/Source/LevelSnapshots/Private/Restorability/SnapshotRestorability.cpp
+++ b/Engine/Plugins/VirtualProduction/LevelSnapshots/Source/LevelSnapshots/Private/Restorability/SnapshotRestorability.cpp
@@ -92,7 +92,7 @@ bool UE::LevelSnapshots::Restorability::IsActorDesirableForCapture(const AActor*
 {
 	SCOPED_SNAPSHOT_CORE_TRACE(IsActorDesirableForCapture);
 
-	if (!IsValid(Actor))
+	if (!IsValid(Actor) || !Private::Internal::DoesActorHaveSupportedClass(Actor))
 	{
 		return false;
 	}
@@ -113,10 +113,8 @@ bool UE::LevelSnapshots::Restorability::IsActorDesirableForCapture(const AActor*
 	{
 		return true;
 	}
-
 	
-	return Private::Internal::DoesActorHaveSupportedClass(Actor)
-            && !Actor->IsTemplate()								// Should never happen, but we never want CDOs
+	return !Actor->IsTemplate()								// Should never happen, but we never want CDOs
 			&& !Actor->HasAnyFlags(RF_Transient)				// Don't add transient actors in non-play worlds	
 #if WITH_EDITOR
             && Actor->IsEditable()
diff --git a/Engine/Plugins/VirtualProduction/LevelSnapshots/Source/LevelSnapshots/Private/Selection/ApplySnapshotFilter.cpp b/Engine/Plugins/VirtualProduction/LevelSnapshots/Source/LevelSnapshots/Private/Selection/ApplySnapshotFilter.cpp
--- a/Engine/Plugins/VirtualProduction/LevelSnapshots/Source/LevelSnapshots/Private/Selection/ApplySnapshotFilter.cpp
+++ b/Engine/Plugins/VirtualProduction/LevelSnapshots/Source/LevelSnapshots/Private/Selection/ApplySnapshotFilter.cpp
@@ -9,17 +9,17 @@
 #include "LevelSnapshotsLog.h"
 #include "LevelSnapshotsModule.h"
 #include "Params/PropertyComparisonParams.h"
+#include "Restorability/SnapshotRestorability.h"
 #include "Selection/PropertySelection.h"
 #include "SnapshotConsoleVariables.h"
 #include "SnapshotCustomVersion.h"
-#include "Restorability/SnapshotRestorability.h"
+#include "Util/EquivalenceUtil.h"
+#include "Util/WorldData/SnapshotObjectUtil.h"
 
 #include "Components/ActorComponent.h"
 #include "GameFramework/Actor.h"
 #include "Modules/ModuleManager.h"
 #include "Stats/StatsMisc.h"
-#include "Util/EquivalenceUtil.h"
-#include "Util/SnapshotObjectUtil.h"
 
 namespace UE::LevelSnapshots::Private::Internal
 {
diff --git a/Engine/Plugins/VirtualProduction/LevelSnapshots/Source/LevelSnapshots/Private/Selection/PropertySelection.cpp b/Engine/Plugins/VirtualProduction/LevelSnapshots/Source/LevelSnapshots/Private/Selection/PropertySelection.cpp
--- a/Engine/Plugins/VirtualProduction/LevelSnapshots/Source/LevelSnapshots/Private/Selection/PropertySelection.cpp
+++ b/Engine/Plugins/VirtualProduction/LevelSnapshots/Source/LevelSnapshots/Private/Selection/PropertySelection.cpp
@@ -28,6 +28,18 @@ FLevelSnapshotPropertyChain::FLevelSnapshotPropertyChain(const FProperty* RootPr
 	AppendInline(RootProperty);
 }
 
+FLevelSnapshotPropertyChain::FLevelSnapshotPropertyChain(const FArchiveSerializedPropertyChain* ChainToLeaf, const FProperty* LeafProperty)
+{
+	if (ChainToLeaf)
+	{
+		for (int32 i = 0; i < ChainToLeaf->GetNumProperties(); ++i)
+		{
+			AppendInline(ChainToLeaf->GetPropertyFromRoot(i));
+		}
+	}
+	AppendInline(LeafProperty);
+}
+
 FLevelSnapshotPropertyChain FLevelSnapshotPropertyChain::MakeAppended(const FProperty* Property) const
 {
 	FLevelSnapshotPropertyChain Result = *this;
diff --git a/Engine/Plugins/VirtualProduction/LevelSnapshots/Source/LevelSnapshots/Private/Selection/PropertySelectionMap.cpp b/Engine/Plugins/VirtualProduction/LevelSnapshots/Source/LevelSnapshots/Private/Selection/PropertySelectionMap.cpp
--- a/Engine/Plugins/VirtualProduction/LevelSnapshots/Source/LevelSnapshots/Private/Selection/PropertySelectionMap.cpp
+++ b/Engine/Plugins/VirtualProduction/LevelSnapshots/Source/LevelSnapshots/Private/Selection/PropertySelectionMap.cpp
@@ -3,6 +3,7 @@
 #include "Selection/PropertySelectionMap.h"
 
 #include "LevelSnapshotsLog.h"
+#include "Algo/ForEach.h"
 #include "Selection/RestorableObjectSelection.h"
 
 #include "GameFramework/Actor.h"
@@ -102,11 +103,6 @@ void FPropertySelectionMap::RemoveCustomEditorSubobjectToRecreate(UObject* Edito
 	}
 }
 
-const FPropertySelection* FPropertySelectionMap::GetSelectedProperties(const FSoftObjectPath& WorldObjectPath) const
-{
-	return EditorWorldObjectToSelectedProperties.Find(WorldObjectPath); 
-}
-
 UE::LevelSnapshots::FRestorableObjectSelection FPropertySelectionMap::GetObjectSelection(const FSoftObjectPath& EditorWorldObject) const
 {
 	return UE::LevelSnapshots::FRestorableObjectSelection(EditorWorldObject, *this);
@@ -127,6 +123,26 @@ TArray<FSoftObjectPath> FPropertySelectionMap::GetKeys() const
 	return Result;
 }
 
+void FPropertySelectionMap::ForEachModifiedObject(TFunctionRef<void(const FSoftObjectPath&)> Callback) const
+{
+	Algo::ForEach(EditorWorldObjectToSelectedProperties, [Callback](const TPair<FSoftObjectPath, FPropertySelection>& Pair)
+	{
+		Callback(Pair.Key);
+	});
+	Algo::ForEach(EditorActorToComponentSelection, [this, Callback](const TPair<FSoftObjectPath, UE::LevelSnapshots::FAddedAndRemovedComponentInfo>& Pair)
+	{
+		if (!EditorWorldObjectToSelectedProperties.Contains(Pair.Key))
+		{
+			Callback(Pair.Key);
+		}
+	});
+}
+
+bool FPropertySelectionMap::HasAnyModifiedActors() const
+{
+	return EditorWorldObjectToSelectedProperties.Num() > 0 || EditorActorToComponentSelection.Num() > 0;
+}
+
 bool FPropertySelectionMap::HasChanges(AActor* EditorWorldActor) const
 {
 	using namespace UE::LevelSnapshots;
diff --git a/Engine/Plugins/VirtualProduction/LevelSnapshots/Source/LevelSnapshots/Private/SnapshotConsoleVariables.cpp b/Engine/Plugins/VirtualProduction/LevelSnapshots/Source/LevelSnapshots/Private/SnapshotConsoleVariables.cpp
--- a/Engine/Plugins/VirtualProduction/LevelSnapshots/Source/LevelSnapshots/Private/SnapshotConsoleVariables.cpp
+++ b/Engine/Plugins/VirtualProduction/LevelSnapshots/Source/LevelSnapshots/Private/SnapshotConsoleVariables.cpp
@@ -13,4 +13,6 @@ TAutoConsoleVariable<FString> UE::LevelSnapshots::ConsoleVariables::CVarBreakOnS
 
 TAutoConsoleVariable<bool> UE::LevelSnapshots::ConsoleVariables::CVarLogTimeDiffingMatchedActors(TEXT("LevelSnapshots.LogTimeDiffingMatchedActors"), false, TEXT("Enable logging of how long matched actors take to diff"));
 TAutoConsoleVariable<bool> UE::LevelSnapshots::ConsoleVariables::CVarLogTimeTakingSnapshots(TEXT("LevelSnapshots.LogTimeTakingSnapshots"), false, TEXT("Enable logging of how long actors take to get snapshots taken off."));
+TAutoConsoleVariable<bool> UE::LevelSnapshots::ConsoleVariables::CVarLogTimeRecreatingActors(TEXT("LevelSnapshots.LogTimeRecreatingActors"), false, TEXT("Enable logging of how it takes to recreate each actor (\"Actors to Add\")"));
+TAutoConsoleVariable<bool> UE::LevelSnapshots::ConsoleVariables::CVarLogTimeRestoringMatchedActors(TEXT("LevelSnapshots.LogTimeRestoringMatchedActors"), false, TEXT("Enable logging of how it takes to serialize into matched actors (\"Modified\")."));
 TAutoConsoleVariable<bool> UE::LevelSnapshots::ConsoleVariables::CVarLogSelectionMap(TEXT("LevelSnapshots.LogSelectionMap"), false, TEXT("Logs all filtered properties a filter is applied."));
diff --git a/Engine/Plugins/VirtualProduction/LevelSnapshots/Source/LevelSnapshots/Private/SnapshotConsoleVariables.h b/Engine/Plugins/VirtualProduction/LevelSnapshots/Source/LevelSnapshots/Private/SnapshotConsoleVariables.h
--- a/Engine/Plugins/VirtualProduction/LevelSnapshots/Source/LevelSnapshots/Private/SnapshotConsoleVariables.h
+++ b/Engine/Plugins/VirtualProduction/LevelSnapshots/Source/LevelSnapshots/Private/SnapshotConsoleVariables.h
@@ -21,8 +21,14 @@ namespace UE::LevelSnapshots::ConsoleVariables
 	/** Enable logging of how long matched actors take to diff*/
 	extern TAutoConsoleVariable<bool> CVarLogTimeDiffingMatchedActors;
 
-	/** Enable logging of how long actors take to get snapshots taken off.*/
+	/** Enable logging of how long actors take to get snapshots taken of.*/
 	extern TAutoConsoleVariable<bool> CVarLogTimeTakingSnapshots;
+	
+	/** Enable logging of how it takes to recreate each actor ("Actors to Add").*/
+	extern TAutoConsoleVariable<bool> CVarLogTimeRecreatingActors;
+	
+	/** Enable logging of how it takes to serialize into matched actors ("Modified").*/
+	extern TAutoConsoleVariable<bool> CVarLogTimeRestoringMatchedActors;
 
 	/** Logs all filtered properties whenever a filter is applied */
 	extern TAutoConsoleVariable<bool> CVarLogSelectionMap;
diff --git a/Engine/Plugins/VirtualProduction/LevelSnapshots/Source/LevelSnapshots/Private/Tests/Util/SnapshotTestRunner.cpp b/Engine/Plugins/VirtualProduction/LevelSnapshots/Source/LevelSnapshots/Private/Tests/Util/SnapshotTestRunner.cpp
--- a/Engine/Plugins/VirtualProduction/LevelSnapshots/Source/LevelSnapshots/Private/Tests/Util/SnapshotTestRunner.cpp
+++ b/Engine/Plugins/VirtualProduction/LevelSnapshots/Source/LevelSnapshots/Private/Tests/Util/SnapshotTestRunner.cpp
@@ -2,6 +2,7 @@
 
 #include "SnapshotTestRunner.h"
 
+#include "LevelSnapshotsFilteringLibrary.h"
 #include "Data/LevelSnapshot.h"
 #include "LevelSnapshotsFunctionLibrary.h"
 
@@ -128,7 +129,7 @@ UE::LevelSnapshots::Private::Tests::FSnapshotTestRunner& UE::LevelSnapshots::Pri
 		if (ensure(SnapshotCounterpart))
 		{
 			FPropertySelectionMap SelectedProperties;
-			ULevelSnapshotsFunctionLibrary::ApplyFilterToFindSelectedProperties(Snapshot, SelectedProperties, OriginalActor, SnapshotCounterpart.GetValue(), Filter);
+			ULevelSnapshotsFilteringLibrary::ApplyFilterToFindSelectedProperties(Snapshot, SelectedProperties, OriginalActor, SnapshotCounterpart.GetValue(), Filter);
 			Callback(SelectedProperties);
 		}
 	}, SnapshotId);
diff --git a/Engine/Plugins/VirtualProduction/LevelSnapshots/Source/LevelSnapshots/Public/Data/ActorSnapshotData.h b/Engine/Plugins/VirtualProduction/LevelSnapshots/Source/LevelSnapshots/Public/Data/ActorSnapshotData.h
--- a/Engine/Plugins/VirtualProduction/LevelSnapshots/Source/LevelSnapshots/Public/Data/ActorSnapshotData.h
+++ b/Engine/Plugins/VirtualProduction/LevelSnapshots/Source/LevelSnapshots/Public/Data/ActorSnapshotData.h
@@ -7,6 +7,7 @@
 #include "ComponentSnapshotData.h"
 #include "CustomSerializationData.h"
 #include "ObjectSnapshotData.h"
+#include "SnapshotUtilTypes.h"
 #include "ActorSnapshotData.generated.h"
 
 /** Holds saved actor data. See ActorUtil for operations. */
@@ -14,7 +15,7 @@ USTRUCT()
 struct LEVELSNAPSHOTS_API FActorSnapshotData
 {
 	GENERATED_BODY()
-	
+
 #if WITH_EDITORONLY_DATA
 	/** The label of the actor when it was saved. */
 	UPROPERTY()
@@ -24,7 +25,11 @@ struct LEVELSNAPSHOTS_API FActorSnapshotData
 
 	/** The class the actor had when it was saved. */
 	UPROPERTY()
-	FSoftClassPath ActorClass;
+	FSoftClassPath ActorClass_DEPRECATED;
+
+	/** Valid index to FWorldSnapshotData::ClassData. Use to lookup class and archetype data. */
+	UPROPERTY()
+	int32 ClassIndex = INDEX_NONE;
 
 	/** The actor's serialized data */
 	UPROPERTY()
@@ -51,4 +56,7 @@ struct LEVELSNAPSHOTS_API FActorSnapshotData
 	/** Used to detect changes without loading actor into memory. */
 	UPROPERTY()
 	FActorSnapshotHash Hash;
+
+	// "Type safety" for better self-documenting code
+	UE::LevelSnapshots::FClassDataIndex GetClassDataIndex() const { return ClassIndex; }
 };
diff --git a/Engine/Plugins/VirtualProduction/LevelSnapshots/Source/LevelSnapshots/Public/Data/LevelSnapshot.h b/Engine/Plugins/VirtualProduction/LevelSnapshots/Source/LevelSnapshots/Public/Data/LevelSnapshot.h
--- a/Engine/Plugins/VirtualProduction/LevelSnapshots/Source/LevelSnapshots/Public/Data/LevelSnapshot.h
+++ b/Engine/Plugins/VirtualProduction/LevelSnapshots/Source/LevelSnapshots/Public/Data/LevelSnapshot.h
@@ -13,7 +13,7 @@ UENUM()
 enum class ECachedDiffResult : uint8
 {
 	/** The actor was not yet analysed */
-	NotInitialised,
+	NotInitialized,
 	/** Actor was analysed and had changes */
 	HadChanges,
 	/** Actor was analysed and had no changes */
@@ -70,7 +70,9 @@ public:
 	 *	@param HandleRemovedActor Actor exists in snapshot but not in world. Receives the original actor path.
 	 *	@param HandleAddedActor Actor exists in world but not in snapshot. Receives reference to world actor.
 	 */
+	UE_DEPRECATED(5.1, "Use DiffWorld that accepts FActorConsumer for HandleMatchedActor instead.")
 	void DiffWorld(UWorld* World, FActorPathConsumer HandleMatchedActor, FActorPathConsumer HandleRemovedActor, FActorConsumer HandleAddedActor) const;
+	void DiffWorld(UWorld* World, FActorConsumer HandleMatchedActor, FActorPathConsumer HandleRemovedActor, FActorConsumer HandleAddedActor) const;
 
 	
 	
@@ -92,6 +94,20 @@ public:
 	const FWorldSnapshotData& GetSerializedData() const { return SerializedData; }
 	const FSnapshotDataCache& GetCache() const { return Cache; }
 
+#if !WITH_EDITOR
+	/**
+	 * Indicates that user code is aware that ClearCachedDiffFlag must be called manually.
+	 *
+	 * In editor builds, the transaction system calls
+	 * ClearCachedDiffFlag but in non-editor builds the transaction system does not exist.
+	 * 
+	 * Enables the use of CachedDiffedActors. 
+	 */
+	void EnableDiffCacheSupport() { bIsDiffCacheEnabled = true; }
+#endif
+	/** Clears the diff state of the actor. Designed to be called by the transaction system (or manually in non-editor builds if EnableDiffCacheSupport has been called). */
+	void ResetDiffCacheToUninitialized(TArrayView<AActor*> ModifiedActors);
+	void ResetDiffCacheToUninitialized(UObject* ModifiedObject); 
 	
 	//~ Begin UObject Interface
 	virtual void BeginDestroy() override;
@@ -105,12 +121,6 @@ private:
 	void EnsureWorldInitialised();
 	void DestroyWorld();
 	void ClearCache();
-
-#if WITH_EDITOR
-	/** Clears FActorSnapshotData::bHasBeenDiffed */
-	void ClearCachedDiffFlag(UObject* ModifiedObject);
-#endif
-
 	
 	/** Callback to destroy our world when editor (editor build) or play (game builds) world is destroyed. */
 	FDelegateHandle Handle;
@@ -118,9 +128,12 @@ private:
 	FDelegateHandle OnObjectModifiedHandle;
 
 	
-	/** The world we will be adding temporary actors to */
+	/** The root world, equivalent to the persistent world, we will be adding temporary actors to */
 	UPROPERTY(Transient)
-	UWorld* SnapshotContainerWorld;
+	UWorld* RootSnapshotWorld;
+	/** Sublevels of RootSnapshotWorld */
+	UPROPERTY(Transient)
+	TArray<UWorld*> SnapshotSublevels;
 	
 	
 	/** The saved snapshot data */
@@ -131,12 +144,21 @@ private:
 	UPROPERTY(Transient)
 	FSnapshotDataCache Cache;
 	
-#if WITH_EDITORONLY_DATA
+	/**
+	 * Only used in non-editor builds (should be wrapped in !WITH_EDITORONLY_DATA but UHT does not allow it).
+	 * Indicates that user code is aware that ClearCachedDiffFlag must be called manually.
+	 *
+	 * In editor builds, the transaction system calls
+	 * ClearCachedDiffFlag but in non-editor builds the transaction system does not exist.
+	 * 
+	 * Enables the use of CachedDiffedActors. 
+	 */
+	UPROPERTY(Transient)
+	bool bIsDiffCacheEnabled = false;
 
 	/** Caches whether an actor was diffed already */	
 	UPROPERTY(Transient)
 	TMap<TWeakObjectPtr<AActor>, ECachedDiffResult> CachedDiffedActors;
-#endif
 
 	
 	/** Path of the map that the snapshot was taken in */
diff --git a/Engine/Plugins/VirtualProduction/LevelSnapshots/Source/LevelSnapshots/Public/Data/ObjectSnapshotData.h b/Engine/Plugins/VirtualProduction/LevelSnapshots/Source/LevelSnapshots/Public/Data/ObjectSnapshotData.h
--- a/Engine/Plugins/VirtualProduction/LevelSnapshots/Source/LevelSnapshots/Public/Data/ObjectSnapshotData.h
+++ b/Engine/Plugins/VirtualProduction/LevelSnapshots/Source/LevelSnapshots/Public/Data/ObjectSnapshotData.h
@@ -14,6 +14,7 @@ struct LEVELSNAPSHOTS_API FObjectSnapshotData
 	UPROPERTY()
 	TArray<uint8> SerializedData;
 
+	/** Flags of the object that was serialized */
 	UPROPERTY()
 	uint64 ObjectFlags {};
 
diff --git a/Engine/Plugins/VirtualProduction/LevelSnapshots/Source/LevelSnapshots/Public/Data/SnapshotDataCache.h b/Engine/Plugins/VirtualProduction/LevelSnapshots/Source/LevelSnapshots/Public/Data/SnapshotDataCache.h
--- a/Engine/Plugins/VirtualProduction/LevelSnapshots/Source/LevelSnapshots/Public/Data/SnapshotDataCache.h
+++ b/Engine/Plugins/VirtualProduction/LevelSnapshots/Source/LevelSnapshots/Public/Data/SnapshotDataCache.h
@@ -4,6 +4,7 @@
 
 #include "CoreMinimal.h"
 #include "UObject/SoftObjectPath.h"
+#include "WorldSnapshotData.h"
 #include "SnapshotDataCache.generated.h"
 
 USTRUCT()
@@ -50,7 +51,7 @@ struct FClassDefaultSnapshotCache
 	GENERATED_BODY()
 
 	UPROPERTY()
-	UObject* CachedLoadedClassDefault = nullptr;
+	TObjectPtr<UObject> CachedLoadedClassDefault = nullptr;
 };
 
 /** Caches data for re-use. */
@@ -65,13 +66,36 @@ struct FSnapshotDataCache
 	UPROPERTY()
 	TMap<FSoftObjectPath, FSubobjectSnapshotCache> SubobjectCache;
 
+	/**
+	 * Caches each archetype object for faster lookup.
+	 * 
+	 * Equal length as FWorldSnapshotData::ClassData and the indices correspond to another.
+	 * Elements can be null in which case there is no cache, yet.
+	 */
 	UPROPERTY()
-	TMap<FSoftClassPath, FClassDefaultSnapshotCache> ClassDefaultCache;
+	TArray<TObjectPtr<UObject>> ArchetypeObjects;
+
+	/**
+	 * Fallback for ClassData. Equal length as ClassData and the indices correspond to another.
+	 * When we fail to archetype data in ClassData, we look for the most appropriate archetype and cache it here.
+	 * 
+	 * Archetype data may be missing because this snapshot is old (was not saved) or because it was explicitly skipped
+	 * by user config.
+	 */
+	TArray<TOptional<FClassSnapshotData>> FallbackArchetypeData;
+
+	void InitFor(const FWorldSnapshotData& SnapshotData)
+	{
+		const int32 Size = SnapshotData.ClassData.Num();
+		ArchetypeObjects.SetNum(Size);
+		FallbackArchetypeData.SetNum(Size);
+	}
 
 	void Reset()
 	{
 		ActorCache.Reset();
 		SubobjectCache.Reset();
-		ClassDefaultCache.Reset();
+		ArchetypeObjects.Reset();
+		FallbackArchetypeData.Reset();
 	}
 };
diff --git a/Engine/Plugins/VirtualProduction/LevelSnapshots/Source/LevelSnapshots/Public/Data/SnapshotUtilTypes.h b/Engine/Plugins/VirtualProduction/LevelSnapshots/Source/LevelSnapshots/Public/Data/SnapshotUtilTypes.h
new file mode 100644
--- /dev/null
+++ b/Engine/Plugins/VirtualProduction/LevelSnapshots/Source/LevelSnapshots/Public/Data/SnapshotUtilTypes.h
@@ -0,0 +1,14 @@
+// Copyright Epic Games, Inc. All Rights Reserved.
+
+#pragma once
+
+#include "CoreMinimal.h"
+
+namespace UE::LevelSnapshots
+{
+	// TODO:
+	/*using FSerializedObjectReferenceIndex = uint32;
+	using FSerializedNameIndex = uint32;*/
+	/** An index to FWorldSnapshotData::ClassData. */
+	using FClassDataIndex = uint32;
+}
\ No newline at end of file
diff --git a/Engine/Plugins/VirtualProduction/LevelSnapshots/Source/LevelSnapshots/Public/Data/SubobjectSnapshotData.h b/Engine/Plugins/VirtualProduction/LevelSnapshots/Source/LevelSnapshots/Public/Data/SubobjectSnapshotData.h
--- a/Engine/Plugins/VirtualProduction/LevelSnapshots/Source/LevelSnapshots/Public/Data/SubobjectSnapshotData.h
+++ b/Engine/Plugins/VirtualProduction/LevelSnapshots/Source/LevelSnapshots/Public/Data/SubobjectSnapshotData.h
@@ -4,6 +4,7 @@
 
 #include "CoreMinimal.h"
 #include "ObjectSnapshotData.h"
+#include "SnapshotUtilTypes.h"
 #include "SubobjectSnapshotData.generated.h"
 
 /** Data saved for subobjects, such as components. */
@@ -11,7 +12,7 @@ USTRUCT()
 struct LEVELSNAPSHOTS_API FSubobjectSnapshotData : public FObjectSnapshotData
 {
 	GENERATED_BODY()
-
+	
 	static FSubobjectSnapshotData MakeSkippedSubobjectData()
 	{
 		FSubobjectSnapshotData Result;
@@ -24,9 +25,16 @@ struct LEVELSNAPSHOTS_API FSubobjectSnapshotData : public FObjectSnapshotData
 	int32 OuterIndex = INDEX_NONE;
 
 	UPROPERTY()
-	FSoftClassPath Class;
+	FSoftClassPath Class_DEPRECATED;
+	
+	/** Valid index to FWorldSnapshotData::ClassData. Use to lookup class and archetype data. */
+    UPROPERTY()
+    int32 ClassIndex = INDEX_NONE;
 
-	/** Whether this class was marked as unsupported when the snapshot was taken */
+	/** Whether this object was marked as unsupported when the snapshot was taken */
 	UPROPERTY()
 	bool bWasSkippedClass = false;
+	
+	// "Type safety" for better self-documenting code
+	UE::LevelSnapshots::FClassDataIndex GetClassDataIndex() const { return ClassIndex; }
 };
\ No newline at end of file
diff --git a/Engine/Plugins/VirtualProduction/LevelSnapshots/Source/LevelSnapshots/Public/Data/WorldSnapshotData.h b/Engine/Plugins/VirtualProduction/LevelSnapshots/Source/LevelSnapshots/Public/Data/WorldSnapshotData.h
--- a/Engine/Plugins/VirtualProduction/LevelSnapshots/Source/LevelSnapshots/Public/Data/WorldSnapshotData.h
+++ b/Engine/Plugins/VirtualProduction/LevelSnapshots/Source/LevelSnapshots/Public/Data/WorldSnapshotData.h
@@ -6,6 +6,7 @@
 #include "ActorSnapshotData.h"
 #include "ComponentSnapshotData.h"
 #include "ClassDefaultObjectSnapshotData.h"
+#include "ClassSnapshotData.h"
 #include "CustomSerializationData.h"
 #include "SnapshotVersion.h"
 #include "SubobjectSnapshotData.h"
@@ -32,9 +33,12 @@ struct LEVELSNAPSHOTS_API FWorldSnapshotData
 	void PostSerialize(const FArchive& Ar);
 	//~ End TStructOpsTypeTraits Interface
 	
-	/* The world we will be adding temporary actors to */
+	/* The root world we will be adding deserialized snapshots actors to */
 	UPROPERTY(Transient)
 	TWeakObjectPtr<UWorld> SnapshotWorld;
+	/** Sublevels of SnapshotWorld; excludes SnapshotWorld->PersistentLevel */
+	UPROPERTY(Transient)
+	TArray<UWorld*> SnapshotSublevels;
 
 	/**
 	 * Stores versioning information we inject into archives.
@@ -50,7 +54,19 @@ struct LEVELSNAPSHOTS_API FWorldSnapshotData
 	 * Because of this, we need to save class defaults in the snapshot.
 	 */
 	UPROPERTY()
-	TMap<FSoftClassPath, FClassDefaultObjectSnapshotData> ClassDefaults;
+	TMap<FSoftClassPath, FClassDefaultObjectSnapshotData> ClassDefaults_DEPRECATED;
+
+	/**
+	 * Saves class info, such as archetype data, for every object's class.
+	 * 
+	 * Actor classes have exactly one entry FClassSnapshotData::FClassSnapshotData.
+	 * For other classes, e.g. components, there may be multiple entries for the same class. This is because
+	 * subobjects can have their own archtetypes (you can see this by looking a the reset to default button in the details panel).
+	 */
+	UPROPERTY()
+	TArray<FClassSnapshotData> ClassData;
+
+
 	
 	/**
 	 * Holds serialized actor data.
@@ -95,6 +111,7 @@ struct LEVELSNAPSHOTS_API FWorldSnapshotData
 	UPROPERTY()
 	TMap<int32, FCustomSerializationData> CustomSubobjectSerializationData;
 
+	
 
 	/** Binds every entry in SerializedNames to its index. Speeds up adding unique names. */
 	UPROPERTY(Transient)
@@ -103,6 +120,10 @@ struct LEVELSNAPSHOTS_API FWorldSnapshotData
 	/** Binds every entry in SerializedObjectReferences to its index. Speeds up adding unique references. */
     UPROPERTY(Transient)
     TMap<FSoftObjectPath, int32> ReferenceToIndex;
+
+	/** Binds every archetype object for which we saved the class data in ClassData to its index. Speeds up adding unique classes. */
+	UPROPERTY(Transient)
+	TMap<TObjectPtr<UObject>, uint32> ArchetypeToClassDataIndex;
 };
 
 template<>
diff --git a/Engine/Plugins/VirtualProduction/LevelSnapshots/Source/LevelSnapshots/Public/LevelSnapshotsFunctionLibrary.h b/Engine/Plugins/VirtualProduction/LevelSnapshots/Source/LevelSnapshots/Public/LevelSnapshotsFunctionLibrary.h
--- a/Engine/Plugins/VirtualProduction/LevelSnapshots/Source/LevelSnapshots/Public/LevelSnapshotsFunctionLibrary.h
+++ b/Engine/Plugins/VirtualProduction/LevelSnapshots/Source/LevelSnapshots/Public/LevelSnapshotsFunctionLibrary.h
@@ -29,6 +29,7 @@ public:
 	 * Goes through the properties of the actors and their components calling IsPropertyValid on them.
 	 * This function does not recursively check object references, e.g. 'Instanced' uproperties. These properties are currently unsupported by the snapshot framework.
 	 */
+	UE_DEPRECATED(5.1, "Use ULevelSnapshotsFilteringLibrary::ApplyFilterToFindSelectedProperties instead.")
 	static void ApplyFilterToFindSelectedProperties(
 		ULevelSnapshot* Snapshot,
 		FPropertySelectionMap& MapToAddTo,
diff --git a/Engine/Plugins/VirtualProduction/LevelSnapshots/Source/LevelSnapshots/Public/Selection/PropertySelection.h b/Engine/Plugins/VirtualProduction/LevelSnapshots/Source/LevelSnapshots/Public/Selection/PropertySelection.h
--- a/Engine/Plugins/VirtualProduction/LevelSnapshots/Source/LevelSnapshots/Public/Selection/PropertySelection.h
+++ b/Engine/Plugins/VirtualProduction/LevelSnapshots/Source/LevelSnapshots/Public/Selection/PropertySelection.h
@@ -22,6 +22,7 @@ struct LEVELSNAPSHOTS_API FLevelSnapshotPropertyChain : FArchiveSerializedProper
 
 	FLevelSnapshotPropertyChain() = default;
 	FLevelSnapshotPropertyChain(const FProperty* RootProperty);
+	FLevelSnapshotPropertyChain(const FArchiveSerializedPropertyChain* ChainToLeaf, const FProperty* LeafProperty);
 	
 	FLevelSnapshotPropertyChain MakeAppended(const FProperty* Property) const;
 	void AppendInline(const FProperty* Property);
diff --git a/Engine/Plugins/VirtualProduction/LevelSnapshots/Source/LevelSnapshots/Public/Selection/PropertySelectionMap.h b/Engine/Plugins/VirtualProduction/LevelSnapshots/Source/LevelSnapshots/Public/Selection/PropertySelectionMap.h
--- a/Engine/Plugins/VirtualProduction/LevelSnapshots/Source/LevelSnapshots/Public/Selection/PropertySelectionMap.h
+++ b/Engine/Plugins/VirtualProduction/LevelSnapshots/Source/LevelSnapshots/Public/Selection/PropertySelectionMap.h
@@ -61,10 +61,10 @@ struct LEVELSNAPSHOTS_API FPropertySelectionMap
 
     /*************** Queries ***************/
 
-    /** Will become DEPRECATED: Use GetObjectSelection(WorldObjectPath).GetObjectSelection()->GetPropertySelection() instead. */
-    const FPropertySelection* GetSelectedProperties(const FSoftObjectPath& WorldObjectPath) const;
     UE::LevelSnapshots::FRestorableObjectSelection GetObjectSelection(const FSoftObjectPath& EditorWorldObject) const;
     TArray<FSoftObjectPath> GetKeys() const;
+	void ForEachModifiedObject(TFunctionRef<void(const FSoftObjectPath&)> Callback) const;
+	bool HasAnyModifiedActors() const;
 
     bool HasChanges(AActor* EditorWorldActor) const;
     bool HasChanges(UActorComponent* EditorWorldComponent) const;
diff --git a/Engine/Plugins/VirtualProduction/LevelSnapshots/Source/LevelSnapshots/Public/Settings/LevelSnapshotsSettings.h b/Engine/Plugins/VirtualProduction/LevelSnapshots/Source/LevelSnapshots/Public/Settings/LevelSnapshotsSettings.h
--- a/Engine/Plugins/VirtualProduction/LevelSnapshots/Source/LevelSnapshots/Public/Settings/LevelSnapshotsSettings.h
+++ b/Engine/Plugins/VirtualProduction/LevelSnapshots/Source/LevelSnapshots/Public/Settings/LevelSnapshotsSettings.h
@@ -2,6 +2,7 @@
 
 #pragma once
 
+#include "SnapshotCompressionSettings.h"
 #include "HashSettings.h"
 #include "Settings/SkippedClassList.h"
 #include "LevelSnapshotsSettings.generated.h"
@@ -53,4 +54,16 @@ public:
 	 */
 	UPROPERTY(Config, EditAnywhere, Category = "Level Snapshots|Performance", meta = (ClampMin = "1"))
 	FHashSettings HashSettings;
+
+	/**
+	 * Specifies the compression settings to use when a snapshot is saved.
+	 *
+	 * These settings affect the time:
+	 * - taken to save a snapshot
+	 * - taken to load a snapshot
+	 *
+	 * You can safely change these settings even if you're already saved snapshots with different settings.
+	 */
+	UPROPERTY(Config, EditAnywhere, Category = "Level Snapshots|Compression", meta = (ClampMin = "1"))
+	FSnapshotCompressionSettings CompressionSettings;
 };
diff --git a/Engine/Plugins/VirtualProduction/LevelSnapshots/Source/LevelSnapshotsEditor/Private/Data/FilterListData.cpp b/Engine/Plugins/VirtualProduction/LevelSnapshots/Source/LevelSnapshotsEditor/Private/Data/FilterListData.cpp
--- a/Engine/Plugins/VirtualProduction/LevelSnapshots/Source/LevelSnapshotsEditor/Private/Data/FilterListData.cpp
+++ b/Engine/Plugins/VirtualProduction/LevelSnapshots/Source/LevelSnapshotsEditor/Private/Data/FilterListData.cpp
@@ -4,6 +4,7 @@
 
 #include "LevelSnapshot.h"
 #include "LevelSnapshotFilters.h"
+#include "LevelSnapshotsFilteringLibrary.h"
 #include "LevelSnapshotsFunctionLibrary.h"
 #include "LevelSnapshotsLog.h"
 
@@ -21,128 +22,37 @@ void FFilterListData::UpdateFilteredList(UWorld* World, ULevelSnapshot* FromSnap
 	FScopedSlowTask DiffDeserializedActors(ExpectedAmountOfWork, LOCTEXT("DiffingActorsKey", "Diffing actors"));
 	DiffDeserializedActors.MakeDialogDelayed(1.f);
 
-    RelatedSnapshot = FromSnapshot;
-
-	// We expect the number of filtered actors & components to stay roughly the same: retain existing memory
-	ModifiedEditorObjectsSelectedProperties_AllowedByFilter.Empty(false);
-	ModifiedWorldActors_AllowedByFilter.Empty(ModifiedWorldActors_AllowedByFilter.Num());
-	RemovedOriginalActorPaths_AllowedByFilter.Empty(RemovedOriginalActorPaths_AllowedByFilter.Num());
-	AddedWorldActors_AllowedByFilter.Empty(AddedWorldActors_AllowedByFilter.Num());
-	
-	FromSnapshot->DiffWorld(
-		World,
-		ULevelSnapshot::FActorPathConsumer::CreateRaw(this, &FFilterListData::HandleActorExistsInWorldAndSnapshot, FilterToApply, &DiffDeserializedActors),
-		ULevelSnapshot::FActorPathConsumer::CreateRaw(this, &FFilterListData::HandleActorWasRemovedFromWorld, FilterToApply),
-		ULevelSnapshot::FActorConsumer::CreateRaw(this, &FFilterListData::HandleActorWasAddedToWorld, FilterToApply)
-		);
-}
-
-void FFilterListData::ApplyFilterToFindSelectedProperties(AActor* WorldActor, ULevelSnapshotFilter* FilterToApply)
-{
-	const FPropertySelection* AllowedSelectedProperties = ModifiedEditorObjectsSelectedProperties_AllowedByFilter.GetSelectedProperties(WorldActor);
-	const FPropertySelection* DisallowedSelectedProperties = ModifiedEditorObjectsSelectedProperties_AllowedByFilter.GetSelectedProperties(WorldActor);
-	
-	if (!ensure(WorldActor && FilterToApply) || AllowedSelectedProperties || DisallowedSelectedProperties)
-	{
-		return;
-	}
-	
-	const bool bIsAllowedByFilters = ModifiedWorldActors_AllowedByFilter.Contains(WorldActor);
-	const bool bIsDisallowedByFilters = ModifiedWorldActors_DisallowedByFilter.Contains(WorldActor);
-	if (!ensureMsgf(bIsAllowedByFilters || bIsDisallowedByFilters, TEXT("You have to call UpdateFilteredList first")))
-	{
-		return; 
-	}
-
-	const TWeakObjectPtr<AActor> DeserializedActor = GetSnapshotCounterpartFor(WorldActor);
-	if (!ensureMsgf(DeserializedActor.IsValid(), TEXT("For some reason this actor has no snapshot counterpart... Investigate.")))
-	{
-		return;
-	}
-	
-	if (bIsAllowedByFilters)
-	{
-		ULevelSnapshotsFunctionLibrary::ApplyFilterToFindSelectedProperties(RelatedSnapshot.Get(), ModifiedEditorObjectsSelectedProperties_AllowedByFilter, WorldActor, DeserializedActor.Get(), FilterToApply);
-	}
-	else
-	{
-		ULevelSnapshotsFunctionLibrary::ApplyFilterToFindSelectedProperties(RelatedSnapshot.Get(), ModifiedEditorObjectsSelectedProperties_DisallowedByFilter, WorldActor, DeserializedActor.Get(), FilterToApply);
-	}
+	RelatedSnapshot = FromSnapshot;
+	SelectionMap = ULevelSnapshotsFilteringLibrary::DiffAndFilterSnapshot(World, FromSnapshot, FilterToApply);
 }
 
-TWeakObjectPtr<AActor> FFilterListData::GetSnapshotCounterpartFor(TWeakObjectPtr<AActor> WorldActor) const
+TWeakObjectPtr<AActor> FFilterListData::GetSnapshotCounterpartFor(const AActor* WorldActor) const
 {
-	const TOptional<TNonNullPtr<AActor>> DeserializedActor = RelatedSnapshot->GetDeserializedActor(WorldActor.Get());
+	const TOptional<TNonNullPtr<AActor>> DeserializedActor = RelatedSnapshot->GetDeserializedActor(WorldActor);
 	return ensureAlwaysMsgf(DeserializedActor, TEXT("Deserialized actor does no exist. Either the snapshots's container world was deleted or the snapshot has no counterpart for this actor"))
 		? DeserializedActor.GetValue() : nullptr;
 }
 
-void FFilterListData::HandleActorExistsInWorldAndSnapshot(const FSoftObjectPath& OriginalActorPath, ULevelSnapshotFilter* FilterToApply, FScopedSlowTask* Progress)
+void FFilterListData::ForEachModifiedActor(TFunctionRef<void(AActor*)> Callback) const
 {
-	Progress->EnterProgressFrame();
-	
-	UObject* ResolvedWorldActor = OriginalActorPath.ResolveObject();
-	if (!ResolvedWorldActor)
+	TSet<AActor*> AlreadyVisited;
+	SelectionMap.ForEachModifiedObject([&AlreadyVisited, Callback](const FSoftObjectPath& ModifiedObject)
 	{
-		UE_LOG(LogLevelSnapshots, Warning, TEXT("Failed to resolve actor %s. Was it deleted from the world?"), *OriginalActorPath.ToString());
-		return;
-	}
+		UObject* Resolved = ModifiedObject.ResolveObject();
+		AActor* Actor = Cast<AActor>(Resolved);
+		Actor = Actor ? Actor : Resolved->GetTypedOuter<AActor>();
 
-	AActor* WorldActor = Cast<AActor>(ResolvedWorldActor);
-	if (ensureAlwaysMsgf(WorldActor, TEXT("A path that was previously associated with an actor no longer refers to an actor. Something is wrong."))
-		&& RelatedSnapshot->HasChangedSinceSnapshotWasTaken(WorldActor))
-	{
-		const TOptional<TNonNullPtr<AActor>> DeserializedSnapshotActor = RelatedSnapshot->GetDeserializedActor(OriginalActorPath);
-		if (!ensureAlwaysMsgf(DeserializedSnapshotActor.Get(nullptr), TEXT("Failed to get TMap value for key %s. Is the snapshot corrupted?"), *OriginalActorPath.ToString()))
-		{
-			return;
-		}
-			
-		const EFilterResult::Type ActorInclusionResult = FilterToApply->IsActorValid(FIsActorValidParams(DeserializedSnapshotActor.GetValue(), WorldActor));
-		if (EFilterResult::CanInclude(ActorInclusionResult))
+		if (Actor && !AlreadyVisited.Contains(Actor))
 		{
-			ModifiedWorldActors_AllowedByFilter.Add(WorldActor);
+			AlreadyVisited.Add(Actor);
+			Callback(Actor);
 		}
-		else
-		{
-			ModifiedWorldActors_DisallowedByFilter.Add(WorldActor);
-		}
-	}
+	});
 }
 
-void FFilterListData::HandleActorWasRemovedFromWorld(const FSoftObjectPath& OriginalActorPath, ULevelSnapshotFilter* FilterToApply)
-{	
-	const EFilterResult::Type FilterResult = FilterToApply->IsDeletedActorValid(
-		FIsDeletedActorValidParams(
-			OriginalActorPath,
-			[this](const FSoftObjectPath& ObjectPath)
-			{
-				return RelatedSnapshot->GetDeserializedActor(ObjectPath).Get(nullptr);
-			}
-		)
-	);
-	
-	if (EFilterResult::CanInclude(FilterResult))
-	{
-		RemovedOriginalActorPaths_AllowedByFilter.Add(OriginalActorPath);
-	}
-	else
-	{
-		RemovedOriginalActorPaths_DisallowedByFilter.Add(OriginalActorPath);
-	}
-}
-
-void FFilterListData::HandleActorWasAddedToWorld(AActor* WorldActor, ULevelSnapshotFilter* FilterToApply)
+bool FFilterListData::HasAnyModifiedActors() const
 {
-	const EFilterResult::Type FilterResult = FilterToApply->IsAddedActorValid(FIsAddedActorValidParams(WorldActor)); 
-	if (EFilterResult::CanInclude(FilterResult))
-	{
-		AddedWorldActors_AllowedByFilter.Add(WorldActor);
-	}
-	else
-	{
-		AddedWorldActors_DisallowedByFilter.Add(WorldActor);
-	}
+	return SelectionMap.HasAnyModifiedActors();
 }
 
-#undef LOCTEXT_NAMESPACE
\ No newline at end of file
+#undef LOCTEXT_NAMESPACE
