Merge changes from Otter stream to Main

diff --git a/Engine/Plugins/VirtualProduction/LevelSnapshots/Source/LevelSnapshotsEditor/Private/Data/FilterListData.h b/Engine/Plugins/VirtualProduction/LevelSnapshots/Source/LevelSnapshotsEditor/Private/Data/FilterListData.h
--- a/Engine/Plugins/VirtualProduction/LevelSnapshots/Source/LevelSnapshotsEditor/Private/Data/FilterListData.h
+++ b/Engine/Plugins/VirtualProduction/LevelSnapshots/Source/LevelSnapshotsEditor/Private/Data/FilterListData.h
@@ -16,57 +16,21 @@ struct FScopedSlowTask;
 struct FFilterListData
 {
 	void UpdateFilteredList(UWorld* World, ULevelSnapshot* FromSnapshot, ULevelSnapshotFilter* FilterToApply);
-	
-	/**
-	 * Runs IsPropertyValid on all properties of WorldActor.
-	 */
-	void ApplyFilterToFindSelectedProperties(AActor* WorldActor, ULevelSnapshotFilter* FilterToApply);
-
-	/**
-	 * If WorldActor is modified, returns the deserialized actor.
-	 * If WorldActor is not modified, returns itself.
-	 */
-	TWeakObjectPtr<AActor> GetSnapshotCounterpartFor(TWeakObjectPtr<AActor> WorldActor) const;
 
-	const FPropertySelectionMap& GetModifiedEditorObjectsSelectedProperties_AllowedByFilter() const { return ModifiedEditorObjectsSelectedProperties_AllowedByFilter; }
-	const TSet<TWeakObjectPtr<AActor>>& GetModifiedActors_AllowedByFilter() const { return ModifiedWorldActors_AllowedByFilter; }
-	const TSet<FSoftObjectPath>& GetRemovedOriginalActorPaths_AllowedByFilter() const { return RemovedOriginalActorPaths_AllowedByFilter; }
-	const TSet<TWeakObjectPtr<AActor>>& GetAddedWorldActors_AllowedByFilter() const { return AddedWorldActors_AllowedByFilter; }
+	/** If WorldActor is modified, returns the deserialized actor. If WorldActor is not modified, returns itself. */
+	TWeakObjectPtr<AActor> GetSnapshotCounterpartFor(const AActor* WorldActor) const;
 
-	const FPropertySelectionMap& GetModifiedEditorObjectsSelectedProperties_DisallowedByFilter() const { return ModifiedEditorObjectsSelectedProperties_DisallowedByFilter; }
-	const TSet<TWeakObjectPtr<AActor>>& GetModifiedActors_DisallowedByFilter() const { return ModifiedWorldActors_DisallowedByFilter; }
-	const TSet<FSoftObjectPath>& GetRemovedOriginalActorPaths_DisallowedByFilter() const { return RemovedOriginalActorPaths_DisallowedByFilter; }
-	const TSet<TWeakObjectPtr<AActor>>& GetAddedWorldActors_DisallowedByFilter() const { return AddedWorldActors_DisallowedByFilter; }
+	void ForEachModifiedActor(TFunctionRef<void(AActor*)> Callback) const;
+	bool HasAnyModifiedActors() const;
+	
+	const FPropertySelectionMap& GetModifiedEditorObjectsSelectedProperties_AllowedByFilter() const { return SelectionMap; }
+	const TSet<FSoftObjectPath>& GetRemovedOriginalActorPaths_AllowedByFilter() const { return SelectionMap.GetDeletedActorsToRespawn(); }
+	const TSet<TWeakObjectPtr<AActor>>& GetAddedWorldActors_AllowedByFilter() const { return SelectionMap.GetNewActorsToDespawn(); }
 
 private:
-
-	void HandleActorExistsInWorldAndSnapshot(const FSoftObjectPath& OriginalActorPath, ULevelSnapshotFilter* FilterToApply, FScopedSlowTask* Progress);
-	void HandleActorWasRemovedFromWorld(const FSoftObjectPath& OriginalActorPath, ULevelSnapshotFilter* FilterToApply);
-	void HandleActorWasAddedToWorld(AActor* WorldActor, ULevelSnapshotFilter* FilterToApply);
-	
 	
 	TWeakObjectPtr<ULevelSnapshot> RelatedSnapshot = nullptr;
-
 	
 	/* Selected properties for actors allowed by filters. */
-	FPropertySelectionMap ModifiedEditorObjectsSelectedProperties_AllowedByFilter;
-
-	/* Actors to show in filter results panel when "ShowUnchanged = false". */
-	TSet<TWeakObjectPtr<AActor>> ModifiedWorldActors_AllowedByFilter;
-	/* Actors which existed in snapshot but not in the world. Only contains entries that passed filters. */
-	TSet<FSoftObjectPath> RemovedOriginalActorPaths_AllowedByFilter;
-	/* Actors which existed in the world but not in the snapshot. Only contains entries that passed filters. */
-	TSet<TWeakObjectPtr<AActor>> AddedWorldActors_AllowedByFilter;
-
-
-	
-	/* Selected properties for actors disallowed by filters. */
-	FPropertySelectionMap ModifiedEditorObjectsSelectedProperties_DisallowedByFilter;
-
-	/* Actors to show in filter results panel when "ShowUnchanged = true". */
-	TSet<TWeakObjectPtr<AActor>> ModifiedWorldActors_DisallowedByFilter;
-	/* Actors which existed in snapshot but not in the world. Only contains entries that did not pass filters. */
-	TSet<FSoftObjectPath> RemovedOriginalActorPaths_DisallowedByFilter;
-	/* Actors which existed in the world but not in the snapshot. Only contains entries that did not pass filters. */
-	TSet<TWeakObjectPtr<AActor>> AddedWorldActors_DisallowedByFilter;
+	FPropertySelectionMap SelectionMap;
 };
\ No newline at end of file
diff --git a/Engine/Plugins/VirtualProduction/LevelSnapshots/Source/LevelSnapshotsEditor/Private/Views/Results/LevelSnapshotsEditorResultsHelpers.cpp b/Engine/Plugins/VirtualProduction/LevelSnapshots/Source/LevelSnapshotsEditor/Private/Views/Results/LevelSnapshotsEditorResultsHelpers.cpp
--- a/Engine/Plugins/VirtualProduction/LevelSnapshots/Source/LevelSnapshotsEditor/Private/Views/Results/LevelSnapshotsEditorResultsHelpers.cpp
+++ b/Engine/Plugins/VirtualProduction/LevelSnapshots/Source/LevelSnapshotsEditor/Private/Views/Results/LevelSnapshotsEditorResultsHelpers.cpp
@@ -53,9 +53,13 @@ TArray<TFieldPath<FProperty>> FLevelSnapshotsEditorResultsHelpers::LoopOverPrope
 	return PropertyRowsGenerated;
 }
 
-void FLevelSnapshotsEditorResultsHelpers::LoopOverHandleHierarchiesAndCreateRowHierarchy(ELevelSnapshotsObjectType InputType, const TWeakPtr<FPropertyHandleHierarchy>& InputHierarchy,
-	const TWeakPtr<FLevelSnapshotsEditorResultsRow>& InDirectParentRow, FPropertySelectionMap& PropertySelectionMap, const TArray<TFieldPath<FProperty>>& PropertiesThatPassFilter,
-	TArray<TFieldPath<FProperty>>& PropertyRowsGenerated, const TWeakPtr<SLevelSnapshotsEditorResults>& InResultsView, 
+void FLevelSnapshotsEditorResultsHelpers::LoopOverHandleHierarchiesAndCreateRowHierarchy(ELevelSnapshotsObjectType InputType,
+	const TWeakPtr<FPropertyHandleHierarchy>& InputHierarchy,
+	const TWeakPtr<FLevelSnapshotsEditorResultsRow>& InDirectParentRow,
+	FPropertySelectionMap& PropertySelectionMap,
+	const TArray<TFieldPath<FProperty>>& PropertiesThatPassFilter,
+	TArray<TFieldPath<FProperty>>& PropertyRowsGenerated,
+	const TWeakPtr<SLevelSnapshotsEditorResults>& InResultsView, 
 	const TWeakPtr<FPropertyHandleHierarchy>& InHierarchyToSearchForCounterparts)
 {
 	if (!ensureAlwaysMsgf(InputType != ObjectType_None, 
@@ -214,7 +218,7 @@ void FLevelSnapshotsEditorResultsHelpers::LoopOverHandleHierarchiesAndCreateRowH
 				BuildModifiedObjectRow(
 					WorldSubobject, SnapshotSubobject, PropertyEditorModule,
 					PropertySelectionMap,
-					PropertySelectionMap.GetSelectedProperties(WorldSubobject)->GetSelectedLeafProperties(),
+					PropertySelectionMap.GetObjectSelection(WorldSubobject).GetPropertySelection()->GetSelectedLeafProperties(),
 					InDirectParentRow, InResultsView,
 					FText::Format(INVTEXT("{PropertyName} ({ObjectName})"), DisplayName, ObjectName));
 
@@ -580,7 +584,7 @@ void FLevelSnapshotsEditorResultsHelpers::BuildNestedSceneComponentRowsRecursive
 		{
 			USceneComponent* CurrentComponent = InHierarchy.Pin()->Component.Get();
 			
-			const FPropertySelection* PropertySelection = PropertySelectionMap.GetSelectedProperties(CurrentComponent);
+			const FPropertySelection* PropertySelection = PropertySelectionMap.GetObjectSelection(CurrentComponent).GetPropertySelection();
 
 			bHasVisibleComponents = PropertySelection ? true : false;
 
@@ -615,7 +619,7 @@ void FLevelSnapshotsEditorResultsHelpers::BuildNestedSceneComponentRowsRecursive
 		if (ensureAlwaysMsgf(CurrentComponent,
 			TEXT("%hs: CurrentComponent was nullptr. Please check the InHierarchy for valid component."), __FUNCTION__))
 		{
-			const FPropertySelection* PropertySelection = PropertySelectionMap.GetSelectedProperties(CurrentComponent);
+			const FPropertySelection* PropertySelection = PropertySelectionMap.GetObjectSelection(CurrentComponent).GetPropertySelection();
 
 			const TArray<TFieldPath<FProperty>>& PropertiesThatPassFilter =
 				PropertySelection ? PropertySelection->GetSelectedLeafProperties() : TArray<TFieldPath<FProperty>>();
diff --git a/Engine/Plugins/VirtualProduction/LevelSnapshots/Source/LevelSnapshotsEditor/Private/Views/Results/LevelSnapshotsEditorResultsRow.cpp b/Engine/Plugins/VirtualProduction/LevelSnapshots/Source/LevelSnapshotsEditor/Private/Views/Results/LevelSnapshotsEditorResultsRow.cpp
--- a/Engine/Plugins/VirtualProduction/LevelSnapshots/Source/LevelSnapshotsEditor/Private/Views/Results/LevelSnapshotsEditorResultsRow.cpp
+++ b/Engine/Plugins/VirtualProduction/LevelSnapshots/Source/LevelSnapshotsEditor/Private/Views/Results/LevelSnapshotsEditorResultsRow.cpp
@@ -17,6 +17,8 @@
 #include "PropertyInfoHelpers.h"
 #include "Styling/SlateIconFinder.h"
 
+#define LOCTEXT_NAMESPACE "LevelSnapshotsEditor"
+
 TWeakPtr<IPropertyRowGenerator> FRowGeneratorInfo::GetGeneratorObject() const
 {
 	return GeneratorObject;
@@ -109,20 +111,18 @@ void FLevelSnapshotsEditorResultsRow::FlushReferences()
 
 FLevelSnapshotsEditorResultsRow::FLevelSnapshotsEditorResultsRow(
 	const FText InDisplayName,                                     
-	const ELevelSnapshotsEditorResultsRowType InRowType, const ECheckBoxState StartingWidgetCheckboxState, 
-	const TWeakPtr<SLevelSnapshotsEditorResults>& InResultsView, const TWeakPtr<FLevelSnapshotsEditorResultsRow>& InDirectParentRow)
+	const ELevelSnapshotsEditorResultsRowType InRowType,
+	const ECheckBoxState StartingWidgetCheckboxState, 
+	const TWeakPtr<SLevelSnapshotsEditorResults>& InResultsView,
+	const TWeakPtr<FLevelSnapshotsEditorResultsRow>& InDirectParentRow)
+		: RowType(InRowType)
+		, DisplayName(InDisplayName)
+		, Tooltip(DisplayName)
+		, DirectParentRow(InDirectParentRow)
+		, ResultsViewPtr(InResultsView)
 {
-	check(InResultsView.IsValid());
-	ResultsViewPtr = InResultsView;
-	
-	DisplayName = InDisplayName;
-	RowType = InRowType;
+	check(ResultsViewPtr.IsValid());
 	SetWidgetCheckedState(StartingWidgetCheckboxState);
-
-	if (InDirectParentRow.IsValid())
-	{
-		DirectParentRow = InDirectParentRow;
-	}
 }
 
 void FLevelSnapshotsEditorResultsRow::InitHeaderRow(const ELevelSnapshotsEditorResultsTreeViewHeaderType InHeaderType,
@@ -142,6 +142,7 @@ void FLevelSnapshotsEditorResultsRow::InitAddedActorRow(AActor* InAddedActor)
 	WorldObject = InAddedActor;
 
 	ApplyRowStateMemoryIfAvailable();
+	InitTooltipWithObject(InAddedActor);
 }
 
 void FLevelSnapshotsEditorResultsRow::InitRemovedActorRow(const FSoftObjectPath& InRemovedActorPath)
@@ -149,6 +150,7 @@ void FLevelSnapshotsEditorResultsRow::InitRemovedActorRow(const FSoftObjectPath&
 	RemovedActorPath = InRemovedActorPath;
 
 	ApplyRowStateMemoryIfAvailable();
+	InitTooltipWithObject(InRemovedActorPath);
 }
 
 void FLevelSnapshotsEditorResultsRow::InitAddedObjectRow(UObject* InAddedObjectToRemove)
@@ -156,6 +158,7 @@ void FLevelSnapshotsEditorResultsRow::InitAddedObjectRow(UObject* InAddedObjectT
 	WorldObject = InAddedObjectToRemove;
 
 	ApplyRowStateMemoryIfAvailable();
+	InitTooltipWithObject(InAddedObjectToRemove);
 }
 
 void FLevelSnapshotsEditorResultsRow::InitRemovedObjectRow(UObject* InRemovedObjectToAdd)
@@ -163,6 +166,7 @@ void FLevelSnapshotsEditorResultsRow::InitRemovedObjectRow(UObject* InRemovedObj
 	WorldObject = InRemovedObjectToAdd;
 
 	ApplyRowStateMemoryIfAvailable();
+	InitTooltipWithObject(InRemovedObjectToAdd);
 }
 
 void FLevelSnapshotsEditorResultsRow::InitActorRow(
@@ -172,6 +176,7 @@ void FLevelSnapshotsEditorResultsRow::InitActorRow(
 	WorldObject = InWorldActor;
 
 	ApplyRowStateMemoryIfAvailable();
+	InitTooltipWithObject(InWorldActor);
 }
 
 void FLevelSnapshotsEditorResultsRow::InitObjectRow(
@@ -312,7 +317,7 @@ void FLevelSnapshotsEditorResultsRow::GenerateModifiedActorGroupChildren(FProper
 	{
 		if (WorldComponent)
 		{
-			if (const FPropertySelection* PropertySelection = PropertySelectionMap.GetSelectedProperties(WorldComponent))
+			if (const FPropertySelection* PropertySelection = PropertySelectionMap.GetObjectSelection(WorldComponent).GetPropertySelection())
 			{
 				// Get remaining properties after filter
 				if (PropertySelection->GetSelectedLeafProperties().Num())
@@ -372,7 +377,7 @@ void FLevelSnapshotsEditorResultsRow::GenerateModifiedActorGroupChildren(FProper
 	{
 		if (CustomSubObjectPair.Key && CustomSubObjectPair.Value)
 		{
-			if (const FPropertySelection* PropertySelection = PropertySelectionMap.GetSelectedProperties(CustomSubObjectPair.Key))
+			if (const FPropertySelection* PropertySelection = PropertySelectionMap.GetObjectSelection(CustomSubObjectPair.Key).GetPropertySelection())
 			{
 				// Get remaining properties after filter
 				if (PropertySelection->GetSelectedLeafProperties().Num())
@@ -394,7 +399,7 @@ void FLevelSnapshotsEditorResultsRow::GenerateModifiedActorGroupChildren(FProper
 	}
 
 	// Top-level Actor Properties
-	if (const FPropertySelection* PropertySelection = PropertySelectionMap.GetSelectedProperties(GetWorldObject()))
+	if (const FPropertySelection* PropertySelection = PropertySelectionMap.GetObjectSelection(GetWorldObject()).GetPropertySelection())
 	{
 		if (PropertySelection->GetSelectedLeafProperties().Num())
 		{
@@ -581,11 +586,6 @@ FText FLevelSnapshotsEditorResultsRow::GetDisplayName() const
 	return DisplayName;
 }
 
-void FLevelSnapshotsEditorResultsRow::SetDisplayName(const FText InDisplayName)
-{
-	DisplayName = InDisplayName;
-}
-
 const FSlateBrush* FLevelSnapshotsEditorResultsRow::GetIconBrush() const
 {
 	if (RowType == ModifiedComponentGroup)
@@ -1120,7 +1120,16 @@ bool FLevelSnapshotsEditorResultsRow::ShouldRowBeVisible() const
 {
 	const bool bShowUnselectedRows = ResultsViewPtr.IsValid() ? ResultsViewPtr.Pin()->GetShowUnselectedRows() : true;
 	const bool bShouldBeVisibleBasedOnCheckedState = bShowUnselectedRows ? true : GetWidgetCheckedState() != ECheckBoxState::Unchecked;
-	return bShouldBeVisibleBasedOnCheckedState && (GetDoesRowMatchSearchTerms() || HasVisibleChildren());
+	return bShouldBeVisibleBasedOnCheckedState && (bDoesRowMatchSearchTerms || HasVisibleChildren());
+}
+
+void FLevelSnapshotsEditorResultsRow::InitTooltipWithObject(const FSoftObjectPath& RowObject)
+{
+	Tooltip = FText::Format(
+		LOCTEXT("TooltipFmt", "{0} ({1})"),
+		DisplayName,
+		FText::FromString(RowObject.ToString())
+	);
 }
 
 bool FLevelSnapshotsEditorResultsRow::GetShouldCheckboxBeHidden() const
@@ -1153,11 +1162,6 @@ void FLevelSnapshotsEditorResultsRow::SetCachedSearchTerms(const FString& InTerm
 	CachedSearchTerms = InTerms;
 }
 
-bool FLevelSnapshotsEditorResultsRow::GetDoesRowMatchSearchTerms() const
-{
-	return bDoesRowMatchSearchTerms;
-}
-
 void FLevelSnapshotsEditorResultsRow::SetDoesRowMatchSearchTerms(const bool bNewMatch)
 {
 	bDoesRowMatchSearchTerms = bNewMatch;
@@ -1260,3 +1264,5 @@ void FLevelSnapshotsEditorResultsRow::EvaluateAndSetAllParentGroupCheckedStates(
 		ParentRow = PinnedParent->GetDirectParentRow();
 	}
 }
+
+#undef LOCTEXT_NAMESPACE
\ No newline at end of file
diff --git a/Engine/Plugins/VirtualProduction/LevelSnapshots/Source/LevelSnapshotsEditor/Private/Views/Results/LevelSnapshotsEditorResultsRow.h b/Engine/Plugins/VirtualProduction/LevelSnapshots/Source/LevelSnapshotsEditor/Private/Views/Results/LevelSnapshotsEditorResultsRow.h
--- a/Engine/Plugins/VirtualProduction/LevelSnapshots/Source/LevelSnapshotsEditor/Private/Views/Results/LevelSnapshotsEditorResultsRow.h
+++ b/Engine/Plugins/VirtualProduction/LevelSnapshots/Source/LevelSnapshotsEditor/Private/Views/Results/LevelSnapshotsEditorResultsRow.h
@@ -128,8 +128,13 @@ struct FLevelSnapshotsEditorResultsRow final : TSharedFromThis<FLevelSnapshotsEd
 
 	void FlushReferences();
 	
-	FLevelSnapshotsEditorResultsRow(const FText InDisplayName, const ELevelSnapshotsEditorResultsRowType InRowType, const ECheckBoxState StartingWidgetCheckboxState, 
-		const TWeakPtr<SLevelSnapshotsEditorResults>& InResultsView, const TWeakPtr<FLevelSnapshotsEditorResultsRow>& InDirectParentRow = nullptr);
+	FLevelSnapshotsEditorResultsRow(
+		const FText InDisplayName,
+		const ELevelSnapshotsEditorResultsRowType InRowType,
+		const ECheckBoxState StartingWidgetCheckboxState, 
+		const TWeakPtr<SLevelSnapshotsEditorResults>& InResultsView,
+		const TWeakPtr<FLevelSnapshotsEditorResultsRow>& InDirectParentRow = nullptr
+		);
 
 	void InitHeaderRow(
 		const ELevelSnapshotsEditorResultsTreeViewHeaderType InHeaderType, const TArray<FText>& InColumns);
@@ -173,7 +178,7 @@ struct FLevelSnapshotsEditorResultsRow final : TSharedFromThis<FLevelSnapshotsEd
 	const TArray<FText>& GetHeaderColumns() const;
 
 	FText GetDisplayName() const;
-	void SetDisplayName(const FText InDisplayName);
+	FText GetTooltip() const { return Tooltip; }
 
 	const FSlateBrush* GetIconBrush() const;
 	
@@ -272,6 +277,7 @@ private:
 	/* Generic properties */
 	ELevelSnapshotsEditorResultsRowType RowType = SingleProperty;
 	FText DisplayName;
+	FText Tooltip;
 	TArray<FLevelSnapshotsEditorResultsRowPtr> ChildRows;
 	bool bIsTreeViewItemExpanded = false;
 
@@ -299,13 +305,6 @@ private:
 	/* This is a breadcrumb trail of display names used to find or store the state of the row. */
 	FString RowPath;
 
-	/* Returns a string of searchable keywords such as object names, property names, paths or anything else associated with the row that might be useful to search for. */
-	const FString& GetOrCacheSearchTerms();
-
-	/* Use MatchSearchTerms() first to match Search Terms against tokens, then call this method*/
-	bool GetDoesRowMatchSearchTerms() const;
-	void SetDoesRowMatchSearchTerms(const bool bNewMatch);
-
 	/* For removed actor type rows */
 	FSoftObjectPath RemovedActorPath;
 
@@ -331,13 +330,21 @@ private:
 	/* Whether the snapshot and world object properties have the same value */
 	bool bIsCounterpartValueSame = false;
 
-	void EvaluateAndSetAllParentGroupCheckedStates() const;
-
-	/* Evaluates all factors which should make a row visible or invisible but does not set visibility. */
-	bool ShouldRowBeVisible() const;
-
 	bool bShouldCheckboxBeHidden;
 
 	/* Checkbox in widget is bound to this property */
 	ECheckBoxState WidgetCheckedState = ECheckBoxState::Checked;
+
+	/* Returns a string of searchable keywords such as object names, property names, paths or anything else associated with the row that might be useful to search for. */
+	const FString& GetOrCacheSearchTerms();
+
+	/* Use MatchSearchTerms() first to match Search Terms against tokens, then call this method*/
+	void SetDoesRowMatchSearchTerms(const bool bNewMatch);
+	
+	void EvaluateAndSetAllParentGroupCheckedStates() const;
+
+	/* Evaluates all factors which should make a row visible or invisible but does not set visibility. */
+	bool ShouldRowBeVisible() const;
+	
+	void InitTooltipWithObject(const FSoftObjectPath& RowObject);
 };
diff --git a/Engine/Plugins/VirtualProduction/LevelSnapshots/Source/LevelSnapshotsEditor/Private/Views/Results/SLevelSnapshotsEditorResults.cpp b/Engine/Plugins/VirtualProduction/LevelSnapshots/Source/LevelSnapshotsEditor/Private/Views/Results/SLevelSnapshotsEditorResults.cpp
--- a/Engine/Plugins/VirtualProduction/LevelSnapshots/Source/LevelSnapshotsEditor/Private/Views/Results/SLevelSnapshotsEditorResults.cpp
+++ b/Engine/Plugins/VirtualProduction/LevelSnapshots/Source/LevelSnapshotsEditor/Private/Views/Results/SLevelSnapshotsEditorResults.cpp
@@ -520,7 +520,7 @@ void SLevelSnapshotsEditorResults::BuildSelectionSetFromSelectedPropertiesInEach
 			}
 			
 			FPropertySelection CheckedNodeFieldPaths;
-			if (const FPropertySelection* PropertySelection = SelectionMap.GetSelectedProperties(WorldObject))
+			if (const FPropertySelection* PropertySelection = SelectionMap.GetObjectSelection(WorldObject).GetPropertySelection())
 			{
 				// Make a copy of the property selection. If a node is unchecked, we'll remove it from the copy.
 				CheckedNodeFieldPaths = *PropertySelection;
@@ -661,29 +661,19 @@ void SLevelSnapshotsEditorResults::BuildSelectionSetFromSelectedPropertiesInEach
 	// Added actors
 	for (const FLevelSnapshotsEditorResultsRowPtr& Group : TreeViewAddedActorGroupObjects)
 	{
-		if (Group.IsValid())
+		if (Group.IsValid() && Group->GetWidgetCheckedState() != ECheckBoxState::Checked)
 		{
-			if (Group->GetWidgetCheckedState() != ECheckBoxState::Unchecked)
-			{
-				// Add to PropertyMap
-				if (AActor* WorldActor = Cast<AActor>(Group->GetWorldObject()))
-				{
-					PropertySelectionMap.AddNewActorToDespawn(WorldActor);
-				}
-			}
+			AActor* WorldActor = Cast<AActor>(Group->GetWorldObject());
+			PropertySelectionMap.RemoveNewActorToDespawn(WorldActor);
 		}
 	}
 
 	// Removed actors
 	for (const FLevelSnapshotsEditorResultsRowPtr& Group : TreeViewRemovedActorGroupObjects)
 	{
-		if (Group.IsValid())
+		if (Group.IsValid() && Group->GetWidgetCheckedState() != ECheckBoxState::Checked)
 		{
-			if (Group->GetWidgetCheckedState() != ECheckBoxState::Unchecked)
-			{
-				// Add to PropertyMap
-				PropertySelectionMap.AddDeletedActorToRespawn(Group->GetObjectPath());
-			}
+			PropertySelectionMap.RemoveDeletedActorToRespawn(Group->GetObjectPath());
 		}
 	}
 
@@ -842,16 +832,13 @@ void SLevelSnapshotsEditorResults::GenerateTreeView(const bool bSnapshotHasChang
 	FlushMemory(!bSnapshotHasChanged);
 	
 	UFilteredResults* FilteredResults = EditorDataPtr->GetFilterResults(); 
-	const TWeakObjectPtr<ULevelSnapshotFilter>& UserFilters = FilteredResults->GetUserFilters();
-
 	FilteredResults->UpdateFilteredResults(EditorDataPtr->GetEditorWorld());
 
 	FilterListData = FilteredResults->GetFilteredData();
-	
 	SplitterManagerPtr = MakeShared<FLevelSnapshotsEditorResultsSplitterManager>(FLevelSnapshotsEditorResultsSplitterManager());
 
 	// Create root headers
-	if (FilterListData.GetModifiedActors_AllowedByFilter().Num())
+	if (FilterListData.HasAnyModifiedActors())
 	{
 		FLevelSnapshotsEditorResultsRowPtr ModifiedActorsHeader = MakeShared<FLevelSnapshotsEditorResultsRow>(
 			FLevelSnapshotsEditorResultsRow(FText::GetEmpty(), FLevelSnapshotsEditorResultsRow::TreeViewHeader, ECheckBoxState::Checked, SharedThis(this)));
@@ -860,8 +847,8 @@ void SLevelSnapshotsEditorResults::GenerateTreeView(const bool bSnapshotHasChang
 			LOCTEXT("ColumnName_CurrentValue", "Current Value"),
 			LOCTEXT("ColumnName_ValueToRestore", "Value to Restore")
 			});
-
-		if (GenerateTreeViewChildren_ModifiedActors(ModifiedActorsHeader, UserFilters.Get()))
+		
+		if (GenerateTreeViewChildren_ModifiedActors(ModifiedActorsHeader))
 		{
 			TreeViewRootHeaderObjects.Add(ModifiedActorsHeader);
 		}
@@ -912,51 +899,39 @@ void SLevelSnapshotsEditorResults::GenerateTreeView(const bool bSnapshotHasChang
 	ExecuteResultsViewSearchOnAllActors(GetSearchStringFromSearchInputField());
 }
 
-bool SLevelSnapshotsEditorResults::GenerateTreeViewChildren_ModifiedActors(FLevelSnapshotsEditorResultsRowPtr ModifiedActorsHeader, ULevelSnapshotFilter* UserFilters)
+bool SLevelSnapshotsEditorResults::GenerateTreeViewChildren_ModifiedActors(FLevelSnapshotsEditorResultsRowPtr ModifiedActorsHeader)
 {
 	check(ModifiedActorsHeader);
 	
-	const TSet<TWeakObjectPtr<AActor>>& ActorsToConsider = FilterListData.GetModifiedActors_AllowedByFilter();
-	TSet<FSoftObjectPath> EvaluatedObjects;
-
-	for (const TWeakObjectPtr<AActor>& WeakWorldActor : ActorsToConsider)
+	FilterListData.ForEachModifiedActor([this, &ModifiedActorsHeader](AActor* WorldActor)
 	{
-		AActor* WorldActor = WeakWorldActor.IsValid() ? WeakWorldActor.Get() : nullptr;
-		if (!ensure(WorldActor) || !UE::LevelSnapshots::Restorability::IsActorDesirableForCapture(WorldActor))
-		{
-			continue;
-		}
-
-		// Get remaining properties after filter
-		if (UserFilters)
+		if (!UE::LevelSnapshots::Restorability::IsActorDesirableForCapture(WorldActor))
 		{
-			FilterListData.ApplyFilterToFindSelectedProperties(WorldActor, UserFilters);
+			return;
 		}
 
 		const FPropertySelectionMap& ModifiedSelectedActors = FilterListData.GetModifiedEditorObjectsSelectedProperties_AllowedByFilter();
 		if (!ModifiedSelectedActors.HasChanges(WorldActor))
 		{
-			continue;
+			return;
 		}
 
-		TWeakObjectPtr<AActor> WeakSnapshotActor = FilterListData.GetSnapshotCounterpartFor(WorldActor);
-
-		const FString& ActorName = WorldActor->GetActorLabel();
 
 		// Create group
+		const FString& ActorName = WorldActor->GetActorLabel();
 		FLevelSnapshotsEditorResultsRowPtr NewActorGroup = MakeShared<FLevelSnapshotsEditorResultsRow>(
 			FLevelSnapshotsEditorResultsRow(
-				FText::FromString(ActorName), FLevelSnapshotsEditorResultsRow::ModifiedActorGroup, ECheckBoxState::Checked, SharedThis(this), ModifiedActorsHeader));
-		
+				FText::FromString(ActorName),
+				FLevelSnapshotsEditorResultsRow::ModifiedActorGroup, ECheckBoxState::Checked,
+				SharedThis(this),
+				ModifiedActorsHeader)
+				);
+		TWeakObjectPtr<AActor> WeakSnapshotActor = FilterListData.GetSnapshotCounterpartFor(WorldActor);
 		NewActorGroup->InitActorRow(WeakSnapshotActor.IsValid() ? WeakSnapshotActor.Get() : nullptr, WorldActor);
-
 		ModifiedActorsHeader->AddToChildRows(NewActorGroup);
 		
 		TreeViewModifiedActorGroupObjects.Add(NewActorGroup);
 
-		// Cache search terms using the desired leaf properties for each object newly added to ModifiedActorsSelectedProperties this loop
-		FString NewCachedSearchTerms = WorldActor->GetHumanReadableName();
-
 		struct Local
 		{
 			static void RecursivelyIterateOverSubObjectsAndReturnSearchTerms(UObject* InObject, const FPropertySelectionMap& InSelectionMap, FString& NewCachedSearchTerms)
@@ -967,7 +942,7 @@ bool SLevelSnapshotsEditorResults::GenerateTreeViewChildren_ModifiedActors(FLeve
 				{
 					for (UObject* SubObject : SubObjects)
 					{
-						if (const FPropertySelection* SubObjectProperties = InSelectionMap.GetSelectedProperties(SubObject))
+						if (const FPropertySelection* SubObjectProperties = InSelectionMap.GetObjectSelection(SubObject).GetPropertySelection())
 						{
 							NewCachedSearchTerms += " " + SubObject->GetName();
 							for (const TFieldPath<FProperty>& LeafProperty : SubObjectProperties->GetSelectedLeafProperties())
@@ -981,8 +956,10 @@ bool SLevelSnapshotsEditorResults::GenerateTreeViewChildren_ModifiedActors(FLeve
 				}
 			}
 		};
-
-		if (const FPropertySelection* ActorProperties = ModifiedSelectedActors.GetSelectedProperties(WorldActor))
+		
+		// Cache search terms using the desired leaf properties for each object newly added to ModifiedActorsSelectedProperties this loop
+		FString NewCachedSearchTerms = WorldActor->GetHumanReadableName();
+		if (const FPropertySelection* ActorProperties = ModifiedSelectedActors.GetObjectSelection(WorldActor).GetPropertySelection())
 		{
 			for (const TFieldPath<FProperty>& LeafProperty : ActorProperties->GetSelectedLeafProperties())
 			{
@@ -991,9 +968,8 @@ bool SLevelSnapshotsEditorResults::GenerateTreeViewChildren_ModifiedActors(FLeve
 		}
 
 		Local::RecursivelyIterateOverSubObjectsAndReturnSearchTerms(WorldActor, ModifiedSelectedActors, NewCachedSearchTerms);
-		
 		NewActorGroup->SetCachedSearchTerms(NewCachedSearchTerms);
-	}
+	});
 
 	return TreeViewModifiedActorGroupObjects.Num() > 0;
 }
diff --git a/Engine/Plugins/VirtualProduction/LevelSnapshots/Source/LevelSnapshotsEditor/Private/Views/Results/SLevelSnapshotsEditorResults.h b/Engine/Plugins/VirtualProduction/LevelSnapshots/Source/LevelSnapshotsEditor/Private/Views/Results/SLevelSnapshotsEditorResults.h
--- a/Engine/Plugins/VirtualProduction/LevelSnapshots/Source/LevelSnapshotsEditor/Private/Views/Results/SLevelSnapshotsEditorResults.h
+++ b/Engine/Plugins/VirtualProduction/LevelSnapshots/Source/LevelSnapshotsEditor/Private/Views/Results/SLevelSnapshotsEditorResults.h
@@ -126,7 +126,7 @@ private:
 	//  Tree View Implementation
 
 	void GenerateTreeView(const bool bSnapshotHasChanged);
-	bool GenerateTreeViewChildren_ModifiedActors(FLevelSnapshotsEditorResultsRowPtr ModifiedActorsHeader, ULevelSnapshotFilter* UserFilters);
+	bool GenerateTreeViewChildren_ModifiedActors(FLevelSnapshotsEditorResultsRowPtr ModifiedActorsHeader);
 	bool GenerateTreeViewChildren_AddedActors(FLevelSnapshotsEditorResultsRowPtr AddedActorsHeader);
 	bool GenerateTreeViewChildren_RemovedActors(FLevelSnapshotsEditorResultsRowPtr RemovedActorsHeader);
 	
diff --git a/Engine/Plugins/VirtualProduction/LevelSnapshots/Source/LevelSnapshotsEditor/Private/Views/Results/SLevelSnapshotsEditorResultsRow.cpp b/Engine/Plugins/VirtualProduction/LevelSnapshots/Source/LevelSnapshotsEditor/Private/Views/Results/SLevelSnapshotsEditorResultsRow.cpp
--- a/Engine/Plugins/VirtualProduction/LevelSnapshots/Source/LevelSnapshotsEditor/Private/Views/Results/SLevelSnapshotsEditorResultsRow.cpp
+++ b/Engine/Plugins/VirtualProduction/LevelSnapshots/Source/LevelSnapshotsEditor/Private/Views/Results/SLevelSnapshotsEditorResultsRow.cpp
@@ -62,8 +62,7 @@ void SLevelSnapshotsEditorResultsRow::Construct(const FArguments& InArgs, const
 		bIsSinglePropertyInCollection || RowType == FLevelSnapshotsEditorResultsRow::StructInSetOrArray || RowType == FLevelSnapshotsEditorResultsRow::StructInMap || bIsParentRowCheckboxHidden ||
 		(bIsSingleProperty && !bHasAtLeastOneValidNode && !bHasAtLeastOneCustomWidget));
 
-	FText Tooltip = bHasValidHandle ? ItemHandle->GetToolTipText() : PinnedItem->GetDisplayName();
-	
+	FText Tooltip = bHasValidHandle ? ItemHandle->GetToolTipText() : PinnedItem->GetTooltip();
 	if (bIsSinglePropertyInCollection)
 	{
 		Tooltip = FText::Format(LOCTEXT("CollectionDisclaimer", "({0}) Individual members of collections cannot be selected. The whole collection will be restored."), Tooltip);
@@ -76,10 +75,6 @@ void SLevelSnapshotsEditorResultsRow::Construct(const FArguments& InArgs, const
 	{
 		Tooltip = FText::Format(LOCTEXT("ComponentOrderDisclaimer", "Component ({0}): Please note that component order reflects the order in the world, not the snapshot. LevelSnapshots does not alter component order."), Tooltip);
 	}
-	else
-	{
-		Tooltip = bHasValidHandle ? ItemHandle->GetToolTipText() : PinnedItem->GetDisplayName();
-	}
 
 	int32 IndentationDepth = 0;
 	TWeakPtr<FLevelSnapshotsEditorResultsRow> ParentRow = PinnedItem->GetDirectParentRow();
diff --git a/Engine/Plugins/VirtualProduction/LevelSnapshots/Source/LevelSnapshotsEditor/Private/Widgets/SLevelSnapshotsEditorBrowser.cpp b/Engine/Plugins/VirtualProduction/LevelSnapshots/Source/LevelSnapshotsEditor/Private/Widgets/SLevelSnapshotsEditorBrowser.cpp
--- a/Engine/Plugins/VirtualProduction/LevelSnapshots/Source/LevelSnapshotsEditor/Private/Widgets/SLevelSnapshotsEditorBrowser.cpp
+++ b/Engine/Plugins/VirtualProduction/LevelSnapshots/Source/LevelSnapshotsEditor/Private/Widgets/SLevelSnapshotsEditorBrowser.cpp
@@ -282,13 +282,11 @@ TSharedRef<SToolTip> SLevelSnapshotsEditorBrowser::CreateCustomTooltip(FAssetDat
 	constexpr uint32 MaxDescriptionHeight = ThumbnailSize / 3 * 2;
 	
 	const TSharedRef<FAssetThumbnail> AssetThumbnail =
-		MakeShareable(
-			new FAssetThumbnail(
-				AssetData.GetAsset(),
-				ThumbnailSize,
-				ThumbnailSize,
-				UThumbnailManager::Get().GetSharedThumbnailPool()
-			)
+		MakeShared<FAssetThumbnail>(
+			AssetData,
+			ThumbnailSize,
+			ThumbnailSize,
+			UThumbnailManager::Get().GetSharedThumbnailPool()
 		);
 	
 	FAssetThumbnailConfig AssetThumbnailConfig;
