Fix group actors getting removed by accident when all of its children are removed but the group actor receives new children afterwards

------------------------------------------------

#[fyi] Jason.Walter
Fix group actor showing up to be recreated if none of its child actors can be resolved

#[fyi] Jason.Walter
Fix snapshots not saving

#[fyi] Jason.Walter
Renaming or moving an actor will no longer causes snapshots to be corrupted

#[fyi] Jason.Walter

diff --git a/Engine/Plugins/VirtualProduction/LevelSnapshots/Source/FoliageSupport/Private/FoliageSupport/FoliageSupport.cpp b/Engine/Plugins/VirtualProduction/LevelSnapshots/Source/FoliageSupport/Private/FoliageSupport/FoliageSupport.cpp
--- a/Engine/Plugins/VirtualProduction/LevelSnapshots/Source/FoliageSupport/Private/FoliageSupport/FoliageSupport.cpp
+++ b/Engine/Plugins/VirtualProduction/LevelSnapshots/Source/FoliageSupport/Private/FoliageSupport/FoliageSupport.cpp
@@ -212,9 +212,9 @@ void UE::LevelSnapshots::Foliage::Private::FFoliageSupport::PostRecreateActor(AA
 	}
 }
 
-void UE::LevelSnapshots::Foliage::Private::FFoliageSupport::PreRemoveActor(AActor* ActorToRemove)
+void UE::LevelSnapshots::Foliage::Private::FFoliageSupport::PreRemoveActor(const FPreRemoveActorParams& Params)
 {
-	if (AInstancedFoliageActor* FoliageActor = Cast<AInstancedFoliageActor>(ActorToRemove))
+	if (AInstancedFoliageActor* FoliageActor = Cast<AInstancedFoliageActor>(Params.ActorToRemove))
 	{
 		FoliageActor->GetLevel()->InstancedFoliageActor.Reset();
 	}
diff --git a/Engine/Plugins/VirtualProduction/LevelSnapshots/Source/FoliageSupport/Private/FoliageSupport/FoliageSupport.h b/Engine/Plugins/VirtualProduction/LevelSnapshots/Source/FoliageSupport/Private/FoliageSupport/FoliageSupport.h
--- a/Engine/Plugins/VirtualProduction/LevelSnapshots/Source/FoliageSupport/Private/FoliageSupport/FoliageSupport.h
+++ b/Engine/Plugins/VirtualProduction/LevelSnapshots/Source/FoliageSupport/Private/FoliageSupport/FoliageSupport.h
@@ -53,7 +53,7 @@ namespace UE::LevelSnapshots::Foliage::Private
 		virtual void PostApplySnapshotProperties(const FApplySnapshotPropertiesParams& Params) override;
 		virtual void PreRecreateActor(UWorld* World, TSubclassOf<AActor> ActorClass, FActorSpawnParameters& InOutSpawnParameters) override;
 		virtual void PostRecreateActor(AActor* RecreatedActor) override;
-		virtual void PreRemoveActor(AActor* ActorToRemove) override;
+		virtual void PreRemoveActor(const FPreRemoveActorParams& Params) override;
 		virtual void PreRemoveComponent(UActorComponent* ComponentToRemove) override;
 		//~ End IRestorationListener Interface
 	};
diff --git a/Engine/Plugins/VirtualProduction/LevelSnapshots/Source/LevelSnapshots/Private/Data/LevelSnapshot.cpp b/Engine/Plugins/VirtualProduction/LevelSnapshots/Source/LevelSnapshots/Private/Data/LevelSnapshot.cpp
--- a/Engine/Plugins/VirtualProduction/LevelSnapshots/Source/LevelSnapshots/Private/Data/LevelSnapshot.cpp
+++ b/Engine/Plugins/VirtualProduction/LevelSnapshots/Source/LevelSnapshots/Private/Data/LevelSnapshot.cpp
@@ -45,6 +45,7 @@ void ULevelSnapshot::ApplySnapshotToWorld(UWorld* TargetWorld, const FPropertySe
 	};
 	
 	EnsureWorldInitialised();
+	
 	UE::LevelSnapshots::Private::ApplyToWorld(SerializedData, Cache, TargetWorld, GetPackage(), SelectionSet);
 }
 
@@ -181,12 +182,12 @@ namespace UE::LevelSnapshots::Private::Internal
 	}
 }
 
-void ULevelSnapshot::DiffWorld(UWorld* World, FActorPathConsumer HandleMatchedActor, FActorPathConsumer HandleRemovedActor, FActorConsumer HandleAddedActor) const
+void ULevelSnapshot::DiffWorld(UWorld* World, FActorPathConsumer HandleMatchedActor, FActorPathConsumer HandleRemovedActor, FActorConsumer HandleAddedActor)
 {
 	DiffWorld(World, FActorConsumer::CreateLambda([&HandleMatchedActor](AActor* Actor){ HandleMatchedActor.Execute(Actor); }), HandleRemovedActor, HandleAddedActor);
 }
 
-void ULevelSnapshot::DiffWorld(UWorld* World, FActorConsumer HandleMatchedActor, FActorPathConsumer HandleRemovedActor, FActorConsumer HandleAddedActor) const
+void ULevelSnapshot::DiffWorld(UWorld* World, FActorConsumer HandleMatchedActor, FActorPathConsumer HandleRemovedActor, FActorConsumer HandleAddedActor)
 {
 	SCOPED_SNAPSHOT_CORE_TRACE(DiffWorld);
 	
@@ -244,7 +245,7 @@ void ULevelSnapshot::DiffWorld(UWorld* World, FActorConsumer HandleMatchedActor,
 		const bool bShouldLogDiffWorldTimes = UE::LevelSnapshots::ConsoleVariables::CVarLogTimeDiffingMatchedActors.GetValueOnAnyThread();
 		FConditionalSortedScopedLog SortedItems(bShouldLogDiffWorldTimes);
 		
-		SerializedData.ForEachOriginalActor([this, &HandleMatchedActor, &HandleRemovedActor, &HandleAddedActor, &AllActors, &LoadedLevels, Settings, &DebugActorName, &SortedItems](const FSoftObjectPath& OriginalActorPath, const FActorSnapshotData& SavedData)
+		SerializedData.ForEachOriginalActor([this, World, &HandleMatchedActor, &HandleRemovedActor, &HandleAddedActor, &AllActors, &LoadedLevels, Settings, &DebugActorName, &SortedItems](const FSoftObjectPath& OriginalActorPath, const FActorSnapshotData& SavedData)
 		{
 			const FSoftObjectPath LevelPath = OriginalActorPath.GetAssetPathString();
 			if (!LoadedLevels.Contains(LevelPath))
@@ -253,21 +254,24 @@ void ULevelSnapshot::DiffWorld(UWorld* World, FActorConsumer HandleMatchedActor,
 				return;
 			}
 			
+			const FSoftClassPath ClassPath = UE::LevelSnapshots::Private::GetClass(SavedData, SerializedData);
+			UClass* const ActorClass = ClassPath.TryLoadClass<AActor>();
+			if (!ActorClass)
+			{
+				UE_LOG(LogLevel, Warning, TEXT("Cannot find class %s. Saved actor %s will not be restored."), *ClassPath.ToString(), *OriginalActorPath.ToString());
+				return;
+			}
+			
 			UObject* ResolvedActor = OriginalActorPath.ResolveObject();
 			// OriginalActorPath may still resolve to a live actor if it was just removed. We need to check the ULevel::Actors to see whether it was removed.
 			const bool bWasRemovedFromWorld = ResolvedActor == nullptr || !AllActors.Contains(Cast<AActor>(ResolvedActor));
-			// We do not need to call IsActorDesirableForCapture: it was already called when we took this snapshot
 			if (bWasRemovedFromWorld)
 			{
-				HandleRemovedActor.Execute(OriginalActorPath);
-				return;
-			}
-
-			const FSoftClassPath ClassPath = UE::LevelSnapshots::Private::GetClass(SavedData, SerializedData);
-			const UClass* ActorClass = ClassPath.TryLoadClass<AActor>();
-			if (!ActorClass)
-			{
-				UE_LOG(LogLevel, Warning, TEXT("Cannot find class %s. Saved actor %s will not be restored."), *ClassPath.ToString(), *OriginalActorPath.ToString());
+				const UE::LevelSnapshots::FCanRecreateActorParams Params { World, ActorClass, OriginalActorPath, [this, &OriginalActorPath](){ return GetDeserializedActor(OriginalActorPath); } };
+				if (UE::LevelSnapshots::Restorability::ShouldConsiderRemovedActorForRecreation(Params))
+				{
+					HandleRemovedActor.Execute(OriginalActorPath);
+				}
 				return;
 			}
 			if (Settings->SkippedClasses.SkippedClasses.Contains(ActorClass))
@@ -436,6 +440,8 @@ void ULevelSnapshot::DestroyWorld()
 		{
 			Sublevel->CleanupWorld();
 		}
+		// Cleanup world clears RF_Standalone flag...
+		SetFlags(RF_Standalone);
 		
 		RootSnapshotWorld = nullptr;
 		SnapshotSublevels.Reset();
diff --git a/Engine/Plugins/VirtualProduction/LevelSnapshots/Source/LevelSnapshots/Private/Data/Util/SnapshotUtil.h b/Engine/Plugins/VirtualProduction/LevelSnapshots/Source/LevelSnapshots/Private/Data/Util/SnapshotUtil.h
--- a/Engine/Plugins/VirtualProduction/LevelSnapshots/Source/LevelSnapshots/Private/Data/Util/SnapshotUtil.h
+++ b/Engine/Plugins/VirtualProduction/LevelSnapshots/Source/LevelSnapshots/Private/Data/Util/SnapshotUtil.h
@@ -19,6 +19,12 @@ namespace UE::LevelSnapshots::Private
 	/** If path contains an actor, return subobject path to that actor. */
 	TOptional<FSoftObjectPath> ExtractActorFromPath(const FSoftObjectPath& OriginalObjectPath, bool& bIsPathToActorSubobject);
 
+	inline bool IsPathToWorldObject(const FSoftObjectPath& OriginalObjectPath)
+	{
+		bool bDummy;
+		return ExtractActorFromPath(OriginalObjectPath, bDummy).IsSet();
+	}
+
 	/** Given a path such as "/Game/MapName.MapName:PersistentLevel.StaticMeshActor_42.SomeComponent" returns the index of "S" in "SomeComponent" */
 	TOptional<int32> FindDotAfterActorName(const FSoftObjectPath& OriginalObjectPath);
 
diff --git a/Engine/Plugins/VirtualProduction/LevelSnapshots/Source/LevelSnapshots/Private/Data/Util/WorldData/WorldDataUtil.cpp b/Engine/Plugins/VirtualProduction/LevelSnapshots/Source/LevelSnapshots/Private/Data/Util/WorldData/WorldDataUtil.cpp
--- a/Engine/Plugins/VirtualProduction/LevelSnapshots/Source/LevelSnapshots/Private/Data/Util/WorldData/WorldDataUtil.cpp
+++ b/Engine/Plugins/VirtualProduction/LevelSnapshots/Source/LevelSnapshots/Private/Data/Util/WorldData/WorldDataUtil.cpp
@@ -16,7 +16,6 @@
 #include "SnapshotConsoleVariables.h"
 #include "SnapshotRestorability.h"
 #include "Util/SortedScopedLog.h"
-#include "Util/Property/PropertyIterator.h"
 
 #include "Algo/ForEach.h"
 #include "Async/ParallelFor.h"
@@ -40,57 +39,32 @@
 
 namespace UE::LevelSnapshots::Private::Internal
 {
-	static TOptional<FComponentSnapshotData> SnapshotComponent(UActorComponent* OriginalComponent)
-	{
-		if (OriginalComponent->CreationMethod == EComponentCreationMethod::UserConstructionScript)
-		{
-			UE_LOG(LogLevelSnapshots, Warning, TEXT("Components created dynamically in the construction script are not supported (%s). Skipping..."), *OriginalComponent->GetPathName());
-			return {};
-		}
-		FComponentSnapshotData Result;
-		Result.CreationMethod = OriginalComponent->CreationMethod;
-		return Result;
-	}
+	static TArray<AActor*> GetAllActorsIn(UWorld* World);
+
+	static void CaptureSnapshotData(const TArray<AActor*>& ActorsInWorld, FWorldSnapshotData& SnapshotData);
+	static void ComputeActorHashes(const TArray<AActor*>& ActorsInWorld, FWorldSnapshotData& SnapshotData);
+}
+
+FWorldSnapshotData UE::LevelSnapshots::Private::SnapshotWorld(UWorld* World)
+{
+	FScopedSlowTask TakeSnapshotTask(2.f, LOCTEXT("TakeSnapshotKey", "Take snapshot"));
+	TakeSnapshotTask.MakeDialogDelayed(1.f);
+
+	const TArray<AActor*> ActorsInWorld = Internal::GetAllActorsIn(World);
+	FWorldSnapshotData SnapshotData;
+	SnapshotData.SnapshotVersionInfo.Initialize();
 	
-	static FActorSnapshotData SnapshotActor(AActor* OriginalActor, FWorldSnapshotData& WorldData)
-	{
-		const FString NameToSearchFor = ConsoleVariables::CVarBreakOnSnapshotActor.GetValueOnAnyThread();
-		if (!NameToSearchFor.IsEmpty() && OriginalActor->GetName().Contains(NameToSearchFor))
-		{
-			UE_DEBUG_BREAK();
-		}
+	TakeSnapshotTask.EnterProgressFrame(1.f);
+	Internal::CaptureSnapshotData(ActorsInWorld, SnapshotData);
 	
-		FActorSnapshotData Result;
-		Result.ClassIndex = AddClassArchetype(WorldData, OriginalActor);
-		FTakeWorldObjectSnapshotArchive::TakeSnapshot(Result.SerializedActorData, WorldData, OriginalActor);
-		// If external modules registered for custom serialisation, trigger their callbacks
-		TakeSnapshotOfActorCustomSubobjects(OriginalActor, Result.CustomActorSerializationData, WorldData);
-	
-		TInlineComponentArray<UActorComponent*> Components;
-		OriginalActor->GetComponents(Components);
-		for (UActorComponent* Comp : Components)
-		{
-			if (!Restorability::IsComponentDesirableForCapture(Comp))
-			{
-				continue;
-			}
-		
-			TOptional<FComponentSnapshotData> SerializedComponentData = SnapshotComponent(Comp);
-			if (SerializedComponentData)
-			{
-				const int32 ComponentIndex = AddObjectDependency(WorldData, Comp);
-				Result.ComponentData.Add(ComponentIndex, *SerializedComponentData);
-				// If external modules registered for custom serialisation, trigger their callbacks
-				TakeSnapshotForSubobject(Comp, WorldData);
-			}
-		}
+	TakeSnapshotTask.EnterProgressFrame(1.f);
+	Internal::ComputeActorHashes(ActorsInWorld, SnapshotData);
 
-#if WITH_EDITORONLY_DATA
-		Result.ActorLabel = OriginalActor->GetActorLabel();
-#endif
-		return Result;
-	}
+	return MoveTemp(SnapshotData);
+}
 
+namespace UE::LevelSnapshots::Private::Internal
+{
 	static TArray<AActor*> GetAllActorsIn(UWorld* World)
 	{
 		TArray<AActor*> ObjectArray;
@@ -117,6 +91,9 @@
 		// Move temp to force optimization on debug / development builds
 		return MoveTemp(ObjectArray);
 	}
+	
+	static FActorSnapshotData SnapshotActor(AActor* OriginalActor, FWorldSnapshotData& WorldData);
+	static TOptional<FComponentSnapshotData> SnapshotComponent(UActorComponent* OriginalComponent);
 
 	static void CaptureSnapshotData(const TArray<AActor*>& ActorsInWorld, FWorldSnapshotData& SnapshotData)
 	{
@@ -136,6 +113,57 @@
 			}
 		});
 	}
+	
+	static FActorSnapshotData SnapshotActor(AActor* OriginalActor, FWorldSnapshotData& WorldData)
+	{
+		const FString NameToSearchFor = ConsoleVariables::CVarBreakOnSnapshotActor.GetValueOnAnyThread();
+		if (!NameToSearchFor.IsEmpty() && OriginalActor->GetName().Contains(NameToSearchFor))
+		{
+			UE_DEBUG_BREAK();
+		}
+	
+		FActorSnapshotData Result;
+		Result.ClassIndex = AddClassArchetype(WorldData, OriginalActor);
+		FTakeWorldObjectSnapshotArchive::TakeSnapshot(Result.SerializedActorData, WorldData, OriginalActor);
+		// If external modules registered for custom serialisation, trigger their callbacks
+		TakeSnapshotOfActorCustomSubobjects(OriginalActor, Result.CustomActorSerializationData, WorldData);
+	
+		TInlineComponentArray<UActorComponent*> Components;
+		OriginalActor->GetComponents(Components);
+		for (UActorComponent* Comp : Components)
+		{
+			if (!Restorability::IsComponentDesirableForCapture(Comp))
+			{
+				continue;
+			}
+		
+			TOptional<FComponentSnapshotData> SerializedComponentData = SnapshotComponent(Comp);
+			if (SerializedComponentData)
+			{
+				const int32 ComponentIndex = AddObjectDependency(WorldData, Comp);
+				Result.ComponentData.Add(ComponentIndex, *SerializedComponentData);
+				// If external modules registered for custom serialisation, trigger their callbacks
+				TakeSnapshotForSubobject(Comp, WorldData);
+			}
+		}
 
+#if WITH_EDITORONLY_DATA
+		Result.ActorLabel = OriginalActor->GetActorLabel();
+#endif
+		return Result;
+	}
+
+	static TOptional<FComponentSnapshotData> SnapshotComponent(UActorComponent* OriginalComponent)
+	{
+		if (OriginalComponent->CreationMethod == EComponentCreationMethod::UserConstructionScript)
+		{
+			UE_LOG(LogLevelSnapshots, Warning, TEXT("Components created dynamically in the construction script are not supported (%s). Skipping..."), *OriginalComponent->GetPathName());
+			return {};
+		}
+		FComponentSnapshotData Result;
+		Result.CreationMethod = OriginalComponent->CreationMethod;
+		return Result;
+	}
+
 	static void ComputeActorHashes(const TArray<AActor*>& ActorsInWorld, FWorldSnapshotData& SnapshotData)
 	{
@@ -154,22 +182,80 @@
 	}
 }
 
-FWorldSnapshotData UE::LevelSnapshots::Private::SnapshotWorld(UWorld* World)
+namespace UE::LevelSnapshots::Private::Internal
 {
-	FScopedSlowTask TakeSnapshotTask(2.f, LOCTEXT("TakeSnapshotKey", "Take snapshot"));
-	TakeSnapshotTask.MakeDialogDelayed(1.f);
+	static void PreloadClassesForRestore(FWorldSnapshotData& WorldData, const FPropertySelectionMap& SelectionMap);
+	static void ApplyToWorld_HandleRemovingActors(UWorld* WorldToApplyTo, const FPropertySelectionMap& PropertiesToSerialize);
+	static void ApplyToWorld_HandleRecreatingActors(FWorldSnapshotData& WorldData, FSnapshotDataCache& Cache, TSet<AActor*>& EvaluatedActors, UPackage* LocalisationSnapshotPackage, const FPropertySelectionMap& PropertiesToSerialize);
+	static void ApplyToWorld_HandleSerializingMatchingActors(FWorldSnapshotData& WorldData, FSnapshotDataCache& Cache, TSet<AActor*>& EvaluatedActors, const TArray<FSoftObjectPath>& SelectedPaths, UPackage* LocalisationSnapshotPackage, const FPropertySelectionMap& PropertiesToSerialize);
 
-	const TArray<AActor*> ActorsInWorld = Internal::GetAllActorsIn(World);
-	FWorldSnapshotData SnapshotData;
-	SnapshotData.SnapshotVersionInfo.Initialize();
+#if WITH_EDITOR
+	class FScopedEditorSelectionClearer
+	{
+		FSelectionStateOfLevel SelectionStateOfLevel;
+	public:
+		FScopedEditorSelectionClearer()
+		{
+			GEditor->GetSelectionStateOfLevel(SelectionStateOfLevel);
+			GEditor->SelectNone(true, true, false);
+		}
+		
+		~FScopedEditorSelectionClearer()
+		{
+			GEditor->SetSelectionStateOfLevel(SelectionStateOfLevel);
+		}
+	};
+#endif
+}
+
+void UE::LevelSnapshots::Private::ApplyToWorld(FWorldSnapshotData& WorldData, FSnapshotDataCache& Cache, UWorld* WorldToApplyTo, UPackage* LocalisationSnapshotPackage, const FPropertySelectionMap& PropertiesToSerialize)
+{
+	check(WorldToApplyTo);
 	
-	TakeSnapshotTask.EnterProgressFrame(1.f);
-	Internal::CaptureSnapshotData(ActorsInWorld, SnapshotData);
+	const TArray<FSoftObjectPath> SelectedPaths = PropertiesToSerialize.GetKeys();
+	const int32 NumActorsToRecreate = PropertiesToSerialize.GetDeletedActorsToRespawn().Num();
+	const int32 NumMatchingActors = SelectedPaths.Num();
+	FScopedSlowTask ApplyToWorldTask(NumActorsToRecreate + NumMatchingActors, LOCTEXT("ApplyToWorldKey", "Apply to world"));
+	ApplyToWorldTask.MakeDialogDelayed(1.f, true);
 	
-	TakeSnapshotTask.EnterProgressFrame(1.f);
-	Internal::ComputeActorHashes(ActorsInWorld, SnapshotData);
+	// Certain custom Blueprint compilers, such as nDisplay, may reset the transaction context. That would cause a crash.
+	Internal::PreloadClassesForRestore(WorldData, PropertiesToSerialize);
+	
+#if WITH_EDITOR
+	// Deleting components while they are selected can cause a crash
+	Internal::FScopedEditorSelectionClearer RestoreSelection;
+	FScopedTransaction Transaction(LOCTEXT("ApplyToWorldKey", "Apply to world"));
+#endif
 
-	return MoveTemp(SnapshotData);
+	// It's important to call these events after the transaction has been started
+	FLevelSnapshotsModule::GetInternalModuleInstance().OnPreApplySnapshot({ PropertiesToSerialize });
+	ON_SCOPE_EXIT
+	{
+		FLevelSnapshotsModule::GetInternalModuleInstance().OnPostApplySnapshot({ PropertiesToSerialize });
+	};
+
+	// Clear editor world subobject cache from previous ApplyToWorld
+	for (auto SubobjectIt = Cache.SubobjectCache.CreateIterator(); SubobjectIt; ++SubobjectIt)
+	{
+		SubobjectIt->Value.EditorObject.Reset();
+	}
+	
+	Internal::ApplyToWorld_HandleRemovingActors(WorldToApplyTo, PropertiesToSerialize);
+	
+	TSet<AActor*> EvaluatedActors;
+	ApplyToWorldTask.EnterProgressFrame(NumActorsToRecreate);
+	Internal::ApplyToWorld_HandleRecreatingActors(WorldData, Cache, EvaluatedActors, LocalisationSnapshotPackage, PropertiesToSerialize);	
+
+	ApplyToWorldTask.EnterProgressFrame(NumMatchingActors);
+	Internal::ApplyToWorld_HandleSerializingMatchingActors(WorldData, Cache, EvaluatedActors, SelectedPaths, LocalisationSnapshotPackage, PropertiesToSerialize);
+	
+	// If we're in the editor then update the gizmos locations as they can get out of sync if any of the serialized actors were selected
+#if WITH_EDITOR
+	if (GUnrealEd)
+	{
+		GUnrealEd->UpdatePivotLocationForSelection();
+	}
+#endif
 }
 
 namespace UE::LevelSnapshots::Private::Internal
@@ -203,34 +289,29 @@
 #if WITH_EDITOR
 		const TSet<TWeakObjectPtr<AActor>>& ActorsToDespawn = PropertiesToSerialize.GetNewActorsToDespawn();
 		const bool bShouldDespawnActors = ActorsToDespawn.Num() > 0;
-		if (!bShouldDespawnActors || !ensure(GEditor))
+		if (!bShouldDespawnActors || !ensure(GUnrealEd))
 		{
 			return;
 		}
 	
 		// Not sure whether needed. "DELETE" command does in UUnrealEdEngine::Exec_Actor ...
 		FEditorDelegates::OnDeleteActorsBegin.Broadcast();
-
-		// Avoid accidentally deleting other user selected actors
-		GEditor->SelectNone(false, false, false);
 	
 		FLevelSnapshotsModule& Module = FLevelSnapshotsModule::GetInternalModuleInstance();
-		USelection* EdSelectionManager = GEditor->GetSelectedActors();
-		EdSelectionManager->BeginBatchSelectOperation();
-		for (const TWeakObjectPtr<AActor>& ActorToDespawn: ActorsToDespawn)
+		TArray<AActor*> ActorsToDelete;
+		ActorsToDelete.Reserve(ActorsToDespawn.Num());
+		for (const TWeakObjectPtr<AActor>& ActorToDespawn : ActorsToDespawn)
 		{
-			if (ensureMsgf(ActorToDespawn.IsValid(), TEXT("Actor became invalid since selection set was created")))
+			if (ActorToDespawn.IsValid())
 			{
-				EdSelectionManager->Modify();
-				Module.OnPreRemoveActor(ActorToDespawn.Get());
-				GEditor->SelectActor(ActorToDespawn.Get(), /*bSelect =*/true, /*bNotifyForActor =*/false, /*bSelectEvenIfHidden =*/true);
+				Module.OnPreRemoveActor({ PropertiesToSerialize, ActorToDespawn.Get() });
+				ActorsToDelete.Add(ActorToDespawn.Get());
 			}
 		}
-		EdSelectionManager->EndBatchSelectOperation();
 
-		const bool bVerifyDeletionCanHappen = true;
-		const bool bWarnAboutReferences = false;
-		GEditor->edactDeleteSelected(WorldToApplyTo, bVerifyDeletionCanHappen, bWarnAboutReferences, bWarnAboutReferences);
+		constexpr bool bVerifyDeletionCanHappen = true;
+		constexpr bool bWarnAboutReferences = false;
+		GUnrealEd->DeleteActors(ActorsToDelete, WorldToApplyTo, GUnrealEd->GetSelectedActors()->GetElementSelectionSet(), bVerifyDeletionCanHappen, bWarnAboutReferences, bWarnAboutReferences);
 
 		// ... and call the end event like in UUnrealEdEngine
 		FEditorDelegates::OnDeleteActorsEnd.Broadcast();
@@ -244,54 +325,12 @@
 			}
 		}
 #endif
-	}
-
-	static bool HandleNameClash(const FSoftObjectPath& OriginalRemovedActorPath)
-	{
-		UObject* FoundObject = FindObject<UObject>(nullptr, *OriginalRemovedActorPath.ToString());
-		if (!FoundObject)
-		{
-			return true;
-		}
-
-		// If it's not an actor nor component then it's possibly an UObjectRedirector
-		UActorComponent* AsComponent = Cast<UActorComponent>(FoundObject); // Unlikely
-		AActor* AsActor = AsComponent ? AsComponent->GetOwner() : Cast<AActor>(FoundObject);
-		const bool bShouldRename= AsActor != nullptr;
-		const bool bCanDestroy = IsValid(AsActor);
-		if (bCanDestroy)
-		{
-			UE_LOG(LogLevelSnapshots, Verbose, TEXT("Handeling name clash for %s"), *OriginalRemovedActorPath.ToString());
-			
-#if WITH_EDITOR
-			GEditor->SelectActor(AsActor, /*bSelect =*/true, /*bNotifyForActor =*/false, /*bSelectEvenIfHidden =*/true);
-			const bool bVerifyDeletionCanHappen = true;
-			const bool bWarnAboutReferences = false;
-			GEditor->edactDeleteSelected(AsActor->GetWorld(), bVerifyDeletionCanHappen, bWarnAboutReferences, bWarnAboutReferences);
-#else
-			AsActor->Destroy(true, true);
-#endif
-		}
-
-		if (bShouldRename)
-		{
-			const FName NewName = MakeUniqueObjectName(FoundObject->GetOuter(), FoundObject->GetClass());
-			return FoundObject->Rename(*NewName.ToString(), nullptr, REN_NonTransactional | REN_DontCreateRedirectors);
-		}
-
-		return false;
-	}
-
-	static FString ExtractActorName(const FSoftObjectPath& OriginalRemovedActorPath)
-	{
-		const FString& SubObjectPath = OriginalRemovedActorPath.GetSubPathString();
-		const int32 LastDotIndex = SubObjectPath.Find(TEXT("."));
-		// Full string: /Game/MapName.MapName:PersistentLevel.StaticMeshActor_42 SubPath: PersistentLevel.StaticMeshActor_42 
-		checkf(LastDotIndex != INDEX_NONE, TEXT("There should always be at least one dot after PersistentLevel"));
-		const int32 NameLength = SubObjectPath.Len() - LastDotIndex - 1;
-		const FString ActorName = SubObjectPath.Right(NameLength);
-		return ActorName;
+		
+		Module.OnPostRemoveActors({ PropertiesToSerialize });
 	}
 	
+	static bool HandleNameClash(const FSoftObjectPath& OriginalRemovedActorPath);
+	static FString ExtractActorName(const FSoftObjectPath& OriginalRemovedActorPath);
+	
 	static void ApplyToWorld_HandleRecreatingActors(FWorldSnapshotData& WorldData, FSnapshotDataCache& Cache, TSet<AActor*>& EvaluatedActors, UPackage* LocalisationSnapshotPackage, const FPropertySelectionMap& PropertiesToSerialize)
 	{
@@ -399,6 +438,53 @@
 			}
 		}
 	}
+
+	static bool HandleNameClash(const FSoftObjectPath& OriginalRemovedActorPath)
+	{
+		UObject* FoundObject = FindObject<UObject>(nullptr, *OriginalRemovedActorPath.ToString());
+		if (!FoundObject)
+		{
+			return true;
+		}
+
+		// If it's not an actor nor component then it's possibly an UObjectRedirector
+		UActorComponent* AsComponent = Cast<UActorComponent>(FoundObject); // Unlikely
+		AActor* AsActor = AsComponent ? AsComponent->GetOwner() : Cast<AActor>(FoundObject);
+		const bool bShouldRename= AsActor != nullptr;
+		const bool bCanDestroy = IsValid(AsActor);
+		if (bCanDestroy)
+		{
+			UE_LOG(LogLevelSnapshots, Verbose, TEXT("Handeling name clash for %s"), *OriginalRemovedActorPath.ToString());
+			
+#if WITH_EDITOR
+			GEditor->SelectActor(AsActor, /*bSelect =*/true, /*bNotifyForActor =*/false, /*bSelectEvenIfHidden =*/true);
+			const bool bVerifyDeletionCanHappen = true;
+			const bool bWarnAboutReferences = false;
+			GEditor->edactDeleteSelected(AsActor->GetWorld(), bVerifyDeletionCanHappen, bWarnAboutReferences, bWarnAboutReferences);
+#else
+			AsActor->Destroy(true, true);
+#endif
+		}
+
+		if (bShouldRename)
+		{
+			const FName NewName = MakeUniqueObjectName(FoundObject->GetOuter(), FoundObject->GetClass());
+			return FoundObject->Rename(*NewName.ToString(), nullptr, REN_NonTransactional | REN_DontCreateRedirectors);
+		}
+
+		return false;
+	}
+
+	static FString ExtractActorName(const FSoftObjectPath& OriginalRemovedActorPath)
+	{
+		const FString& SubObjectPath = OriginalRemovedActorPath.GetSubPathString();
+		const int32 LastDotIndex = SubObjectPath.Find(TEXT("."));
+		// Full string: /Game/MapName.MapName:PersistentLevel.StaticMeshActor_42 SubPath: PersistentLevel.StaticMeshActor_42 
+		checkf(LastDotIndex != INDEX_NONE, TEXT("There should always be at least one dot after PersistentLevel"));
+		const int32 NameLength = SubObjectPath.Len() - LastDotIndex - 1;
+		const FString ActorName = SubObjectPath.Right(NameLength);
+		return ActorName;
+	}
 	
 	static void ApplyToWorld_HandleSerializingMatchingActors(FWorldSnapshotData& WorldData, FSnapshotDataCache& Cache, TSet<AActor*>& EvaluatedActors, const TArray<FSoftObjectPath>& SelectedPaths, UPackage* LocalisationSnapshotPackage, const FPropertySelectionMap& PropertiesToSerialize)
 	{
@@ -453,68 +539,6 @@
 			ExitingActorTask.EnterProgressFrame(1);
 		}
 	}
-
-#if WITH_EDITOR
-	class FScopedEditorSelectionClearer
-	{
-		FSelectionStateOfLevel SelectionStateOfLevel;
-	public:
-		FScopedEditorSelectionClearer()
-		{
-			GEditor->GetSelectionStateOfLevel(SelectionStateOfLevel);
-			GEditor->SelectNone(true, true, false);
-		}
-		
-		~FScopedEditorSelectionClearer()
-		{
-			GEditor->SetSelectionStateOfLevel(SelectionStateOfLevel);
-		}
-	};
-#endif
 }
 
-void UE::LevelSnapshots::Private::ApplyToWorld(FWorldSnapshotData& WorldData, FSnapshotDataCache& Cache, UWorld* WorldToApplyTo, UPackage* LocalisationSnapshotPackage, const FPropertySelectionMap& PropertiesToSerialize)
-{
-	using namespace UE::LevelSnapshots::Private::Internal;
-	check(WorldToApplyTo);
-	
-	const TArray<FSoftObjectPath> SelectedPaths = PropertiesToSerialize.GetKeys();
-	const int32 NumActorsToRecreate = PropertiesToSerialize.GetDeletedActorsToRespawn().Num();
-	const int32 NumMatchingActors = SelectedPaths.Num();
-	FScopedSlowTask ApplyToWorldTask(NumActorsToRecreate + NumMatchingActors, LOCTEXT("ApplyToWorldKey", "Apply to world"));
-	ApplyToWorldTask.MakeDialogDelayed(1.f, true);
-	
-	// Certain custom Blueprint compilers, such as nDisplay, may reset the transaction context. That would cause a crash.
-	PreloadClassesForRestore(WorldData, PropertiesToSerialize);
-	
-#if WITH_EDITOR
-	// Deleting components while they are selected can cause a crash
-	FScopedEditorSelectionClearer RestoreSelection;
-	FScopedTransaction Transaction(FText::FromString("Loading Level Snapshot."));
-#endif
-
-	// Clear editor world subobject cache from previous ApplyToWorld
-	for (auto SubobjectIt = Cache.SubobjectCache.CreateIterator(); SubobjectIt; ++SubobjectIt)
-	{
-		SubobjectIt->Value.EditorObject.Reset();
-	}
-	
-	ApplyToWorld_HandleRemovingActors(WorldToApplyTo, PropertiesToSerialize);
-	
-	TSet<AActor*> EvaluatedActors;
-	ApplyToWorldTask.EnterProgressFrame(NumActorsToRecreate);
-	ApplyToWorld_HandleRecreatingActors(WorldData, Cache, EvaluatedActors, LocalisationSnapshotPackage, PropertiesToSerialize);	
-
-	ApplyToWorldTask.EnterProgressFrame(NumMatchingActors);
-	ApplyToWorld_HandleSerializingMatchingActors(WorldData, Cache, EvaluatedActors, SelectedPaths, LocalisationSnapshotPackage, PropertiesToSerialize);
-	
-	// If we're in the editor then update the gizmos locations as they can get out of sync if any of the serialized actors were selected
-#if WITH_EDITOR
-	if (GUnrealEd)
-	{
-		GUnrealEd->UpdatePivotLocationForSelection();
-	}
-#endif
-}
-
 #undef LOCTEXT_NAMESPACE
diff --git a/Engine/Plugins/VirtualProduction/LevelSnapshots/Source/LevelSnapshots/Private/Data/WorldSnapshotData.cpp b/Engine/Plugins/VirtualProduction/LevelSnapshots/Source/LevelSnapshots/Private/Data/WorldSnapshotData.cpp
--- a/Engine/Plugins/VirtualProduction/LevelSnapshots/Source/LevelSnapshots/Private/Data/WorldSnapshotData.cpp
+++ b/Engine/Plugins/VirtualProduction/LevelSnapshots/Source/LevelSnapshots/Private/Data/WorldSnapshotData.cpp
@@ -7,6 +7,7 @@
 #include "Util/Property/PropertyIterator.h"
 
 #include "Serialization/BufferArchive.h"
+#include "Util/SnapshotUtil.h"
 #include "Util/WorldData/ClassDataUtil.h"
 #include "Util/WorldData/CompressionUtil.h"
 
@@ -60,11 +61,24 @@ namespace UE::LevelSnapshots::Private::Internal
 		}
 	}
 	
-	static void CollectReferencesAndNames(FWorldSnapshotData& WorldData, FArchive& Ar)
+	static void CollectReferencesAndNames(bool bSkipActorReferences, FWorldSnapshotData& WorldData, FArchive& Ar)
 	{
 		// References
-		Ar << WorldData.SerializedObjectReferences;
-		CollectActorReferences(WorldData, Ar);
+		if (bSkipActorReferences)
+		{
+			for (FSoftObjectPath& Path : WorldData.SerializedObjectReferences)
+			{
+				if (!IsPathToWorldObject(Path))
+				{
+					Ar << Path;
+				}
+			}
+		}
+		else
+		{
+			Ar << WorldData.SerializedObjectReferences;
+			CollectActorReferences(WorldData, Ar);
+		}
 		CollectClassDefaultReferences(WorldData, Ar);
 
 		// Names
@@ -179,11 +193,15 @@ bool FWorldSnapshotData::Serialize(FArchive& Ar)
 {
 	using namespace UE::LevelSnapshots::Private;
 	Ar.UsingCustomVersion(FSnapshotCustomVersion::GUID);
-	
+
 	// When this struct is saved, the save algorithm collects references. It's faster if we just give it the info directly.
 	if (Ar.IsObjectReferenceCollector())
 	{
-		UE::LevelSnapshots::Private::Internal::CollectReferencesAndNames(*this, Ar);
+		// Actor references must not be renamed because a snapshot is supposed to keep track of what was in the level at a given time.
+		// If an actor is renamed or moved to another level, the snapshot should ignore those changes.
+		const bool bIsRenameArchive = !Ar.IsPersistent() && Ar.IsSaving() && Ar.IsModifyingWeakAndStrongReferences();
+		const bool bSkipActorReferences = bIsRenameArchive;
+		UE::LevelSnapshots::Private::Internal::CollectReferencesAndNames(bSkipActorReferences, *this, Ar);
 		return true;
 	}
 
diff --git a/Engine/Plugins/VirtualProduction/LevelSnapshots/Source/LevelSnapshots/Private/LevelSnapshotsModule.cpp b/Engine/Plugins/VirtualProduction/LevelSnapshots/Source/LevelSnapshots/Private/LevelSnapshotsModule.cpp
--- a/Engine/Plugins/VirtualProduction/LevelSnapshots/Source/LevelSnapshots/Private/LevelSnapshotsModule.cpp
+++ b/Engine/Plugins/VirtualProduction/LevelSnapshots/Source/LevelSnapshots/Private/LevelSnapshotsModule.cpp
@@ -79,7 +79,7 @@ void UE::LevelSnapshots::Private::FLevelSnapshotsModule::RegisterRestorabilityOv
 	Overrides.AddUnique(Overrider);
 }
 
-void UE::LevelSnapshots::Private::FLevelSnapshotsModule::UnregisterRestorabilityOverrider(TSharedRef<ISnapshotRestorabilityOverrider> Overrider)
+void UE::LevelSnapshots::Private::FLevelSnapshotsModule::UnregisterRestorabilityOverrider(const TSharedRef<ISnapshotRestorabilityOverrider>& Overrider)
 {
 	Overrides.RemoveSwap(Overrider);
 }
@@ -114,7 +114,7 @@ void UE::LevelSnapshots::Private::FLevelSnapshotsModule::RegisterPropertyCompare
 	PropertyComparers.FindOrAdd(Class).AddUnique(Comparer);
 }
 
-void UE::LevelSnapshots::Private::FLevelSnapshotsModule::UnregisterPropertyComparer(UClass* Class, TSharedRef<IPropertyComparer> Comparer)
+void UE::LevelSnapshots::Private::FLevelSnapshotsModule::UnregisterPropertyComparer(UClass* Class, const TSharedRef<IPropertyComparer>& Comparer)
 {
 	TArray<TSharedRef<IPropertyComparer>>* Comparers = PropertyComparers.Find(Class);
 	if (!Comparers)
@@ -156,12 +156,22 @@ void UE::LevelSnapshots::Private::FLevelSnapshotsModule::UnregisterCustomObjectS
 	CustomSerializers.Remove(Class);
 }
 
+void UE::LevelSnapshots::Private::FLevelSnapshotsModule::RegisterGlobalActorFilter(TSharedRef<IActorSnapshotFilter> Filter)
+{
+	GlobalFilters.AddUnique(MoveTemp(Filter));
+}
+
+void UE::LevelSnapshots::Private::FLevelSnapshotsModule::UnregisterGlobalActorFilter(const TSharedRef<IActorSnapshotFilter>& Filter)
+{
+	GlobalFilters.RemoveSingle(Filter);
+}
+
 void UE::LevelSnapshots::Private::FLevelSnapshotsModule::RegisterSnapshotLoader(TSharedRef<ISnapshotLoader> Loader)
 {
 	SnapshotLoaders.AddUnique(Loader);
 }
 
-void UE::LevelSnapshots::Private::FLevelSnapshotsModule::UnregisterSnapshotLoader(TSharedRef<ISnapshotLoader> Loader)
+void UE::LevelSnapshots::Private::FLevelSnapshotsModule::UnregisterSnapshotLoader(const TSharedRef<ISnapshotLoader>& Loader)
 {
 	SnapshotLoaders.RemoveSingle(Loader);
 }
@@ -171,7 +181,7 @@ void UE::LevelSnapshots::Private::FLevelSnapshotsModule::RegisterRestorationList
 	RestorationListeners.AddUnique(Listener);
 }
 
-void UE::LevelSnapshots::Private::FLevelSnapshotsModule::UnregisterRestorationListener(TSharedRef<IRestorationListener> Listener)
+void UE::LevelSnapshots::Private::FLevelSnapshotsModule::UnregisterRestorationListener(const TSharedRef<IRestorationListener>& Listener)
 {
 	RestorationListeners.RemoveSingle(Listener);
 }
@@ -307,6 +317,42 @@ TSharedPtr<UE::LevelSnapshots::ICustomObjectSnapshotSerializer> UE::LevelSnapsho
 	return nullptr;
 }
 
+bool UE::LevelSnapshots::Private::FLevelSnapshotsModule::CanRecreateActor(const FCanRecreateActorParams& Params) const
+{
+	for (const TSharedRef<IActorSnapshotFilter>& Filter : GlobalFilters)
+	{
+		const IActorSnapshotFilter::EFilterResult Result = Filter->CanRecreateActor(Params);
+		switch (Result)
+		{
+		case IActorSnapshotFilter::EFilterResult::Allow: return true;
+		case IActorSnapshotFilter::EFilterResult::DoNotCare: continue;
+		case IActorSnapshotFilter::EFilterResult::Disallow: return false;
+		default:
+			checkNoEntry();
+		}
+	};
+
+	return true;
+}
+
+bool UE::LevelSnapshots::Private::FLevelSnapshotsModule::CanDeleteActor(const AActor* EditorActor) const
+{
+	for (const TSharedRef<IActorSnapshotFilter>& Filter : GlobalFilters)
+	{
+		const IActorSnapshotFilter::EFilterResult Result = Filter->CanDeleteActor(EditorActor);
+		switch (Result)
+		{
+		case IActorSnapshotFilter::EFilterResult::Allow: return true;
+		case IActorSnapshotFilter::EFilterResult::DoNotCare: continue;
+		case IActorSnapshotFilter::EFilterResult::Disallow: return false;
+		default:
+			checkNoEntry();
+		}
+	};
+
+	return true;
+}
+
 void UE::LevelSnapshots::Private::FLevelSnapshotsModule::AddCanTakeSnapshotDelegate(FName DelegateName, FCanTakeSnapshot Delegate)
 {
 	CanTakeSnapshotDelegates.FindOrAdd(DelegateName) = Delegate;
@@ -340,6 +386,26 @@ void UE::LevelSnapshots::Private::FLevelSnapshotsModule::OnPostLoadSnapshotObjec
 	}
 }
 
+void UE::LevelSnapshots::Private::FLevelSnapshotsModule::OnPreApplySnapshot(const FApplySnapshotParams& Params)
+{
+	SCOPED_SNAPSHOT_CORE_TRACE(RestorationListeners);
+	
+	for (const TSharedRef<IRestorationListener>& Listener : RestorationListeners)
+	{
+		Listener->PreApplySnapshot(Params);
+	}
+}
+
+void UE::LevelSnapshots::Private::FLevelSnapshotsModule::OnPostApplySnapshot(const FApplySnapshotParams& Params)
+{
+	SCOPED_SNAPSHOT_CORE_TRACE(RestorationListeners);
+	
+	for (const TSharedRef<IRestorationListener>& Listener : RestorationListeners)
+	{
+		Listener->PostApplySnapshot(Params);
+	}
+}
+
 void UE::LevelSnapshots::Private::FLevelSnapshotsModule::OnPreApplySnapshotProperties(const FApplySnapshotPropertiesParams& Params)
 {
 	SCOPED_SNAPSHOT_CORE_TRACE(RestorationListeners);
@@ -400,13 +466,26 @@ void UE::LevelSnapshots::Private::FLevelSnapshotsModule::OnPostRecreateActor(AAc
 	}
 }
 
-void UE::LevelSnapshots::Private::FLevelSnapshotsModule::OnPreRemoveActor(AActor* Actor)
+void UE::LevelSnapshots::Private::FLevelSnapshotsModule::OnPreRemoveActor(const FPreRemoveActorParams& Params)
+{
+	SCOPED_SNAPSHOT_CORE_TRACE(RestorationListeners);
+	
+	for (const TSharedRef<IRestorationListener>& Listener : RestorationListeners)
+	{
+		PRAGMA_DISABLE_DEPRECATION_WARNINGS
+		Listener->PreRemoveActor(Params.ActorToRemove);
+		PRAGMA_ENABLE_DEPRECATION_WARNINGS
+		Listener->PreRemoveActor(Params);
+	}
+}
+
+void UE::LevelSnapshots::Private::FLevelSnapshotsModule::OnPostRemoveActors(const FPostRemoveActorsParams& Params)
 {
 	SCOPED_SNAPSHOT_CORE_TRACE(RestorationListeners);
 	
 	for (const TSharedRef<IRestorationListener>& Listener : RestorationListeners)
 	{
-		Listener->PreRemoveActor(Actor);
+		Listener->PostRemoveActors(Params);
 	}
 }
 
diff --git a/Engine/Plugins/VirtualProduction/LevelSnapshots/Source/LevelSnapshots/Private/LevelSnapshotsModule.h b/Engine/Plugins/VirtualProduction/LevelSnapshots/Source/LevelSnapshots/Private/LevelSnapshotsModule.h
--- a/Engine/Plugins/VirtualProduction/LevelSnapshots/Source/LevelSnapshots/Private/LevelSnapshotsModule.h
+++ b/Engine/Plugins/VirtualProduction/LevelSnapshots/Source/LevelSnapshots/Private/LevelSnapshotsModule.h
@@ -29,17 +29,19 @@ namespace UE::LevelSnapshots::Private
 
 		//~ Begin ILevelSnapshotsModule Interface
 		virtual void RegisterRestorabilityOverrider(TSharedRef<ISnapshotRestorabilityOverrider> Overrider) override;
-		virtual void UnregisterRestorabilityOverrider(TSharedRef<ISnapshotRestorabilityOverrider> Overrider) override;
+		virtual void UnregisterRestorabilityOverrider(const TSharedRef<ISnapshotRestorabilityOverrider>& Overrider) override;
 		virtual void AddSkippedSubobjectClasses(const TSet<UClass*>& Classes) override;
 		virtual void RemoveSkippedSubobjectClasses(const TSet<UClass*>& Classes) override;
 		virtual void RegisterPropertyComparer(UClass* Class, TSharedRef<IPropertyComparer> Comparer) override;
-		virtual void UnregisterPropertyComparer(UClass* Class, TSharedRef<IPropertyComparer> Comparer) override;
+		virtual void UnregisterPropertyComparer(UClass* Class, const TSharedRef<IPropertyComparer>& Comparer) override;
 		virtual void RegisterCustomObjectSerializer(UClass* Class, TSharedRef<ICustomObjectSnapshotSerializer> CustomSerializer, bool bIncludeBlueprintChildClasses = true);
 		virtual void UnregisterCustomObjectSerializer(UClass* Class) override;
+		virtual void RegisterGlobalActorFilter(TSharedRef<IActorSnapshotFilter> Filter) override;
+		virtual void UnregisterGlobalActorFilter(const TSharedRef<IActorSnapshotFilter>& Filter) override;
 		virtual void RegisterSnapshotLoader(TSharedRef<ISnapshotLoader> Loader) override;
-		virtual void UnregisterSnapshotLoader(TSharedRef<ISnapshotLoader> Loader) override;
+		virtual void UnregisterSnapshotLoader(const TSharedRef<ISnapshotLoader>& Loader) override;
 		virtual void RegisterRestorationListener(TSharedRef<IRestorationListener> Listener) override;
-		virtual void UnregisterRestorationListener(TSharedRef<IRestorationListener> Listener) override;
+		virtual void UnregisterRestorationListener(const TSharedRef<IRestorationListener>& Listener) override;
 		virtual void AddExplicitilySupportedProperties(const TSet<const FProperty*>& Properties) override;
 		virtual void RemoveAdditionallySupportedProperties(const TSet<const FProperty*>& Properties) override;
 		virtual void AddExplicitlyUnsupportedProperties(const TSet<const FProperty*>& Properties) override;
@@ -60,12 +62,18 @@ namespace UE::LevelSnapshots::Private
 
 		TSharedPtr<ICustomObjectSnapshotSerializer> GetCustomSerializerForClass(UClass* Class) const;
 
+		bool CanRecreateActor(const FCanRecreateActorParams& Params) const;
+		bool CanDeleteActor(const AActor* EditorActor) const;
+
 		virtual void AddCanTakeSnapshotDelegate(FName DelegateName, FCanTakeSnapshot Delegate) override;
 		virtual void RemoveCanTakeSnapshotDelegate(FName DelegateName) override;
 		virtual bool CanTakeSnapshot(const FPreTakeSnapshotEventData& Event) const override;
 
 		void OnPostLoadSnapshotObject(const FPostLoadSnapshotObjectParams& Params);
 
+		void OnPreApplySnapshot(const FApplySnapshotParams& Params);
+		void OnPostApplySnapshot(const FApplySnapshotParams& Params);
+		
 		void OnPreApplySnapshotProperties(const FApplySnapshotPropertiesParams& Params);
 		void OnPostApplySnapshotProperties(const FApplySnapshotPropertiesParams& Params);
 
@@ -75,7 +83,8 @@ namespace UE::LevelSnapshots::Private
 		void OnPreCreateActor(UWorld* World, TSubclassOf<AActor> ActorClass, FActorSpawnParameters& InOutSpawnParams);
 		void OnPostRecreateActor(AActor* Actor);
 		
-		void OnPreRemoveActor(AActor* Actor);
+		void OnPreRemoveActor(const FPreRemoveActorParams& Params);
+		void OnPostRemoveActors(const FPostRemoveActorsParams& Params);
 		
 		void OnPreRecreateComponent(const FPreRecreateComponentParams& Params);
 		void OnPostRecreateComponent(UActorComponent* RecreatedComponent);
@@ -99,7 +108,8 @@ namespace UE::LevelSnapshots::Private
 		
 		TMap<FSoftClassPath, TArray<TSharedRef<IPropertyComparer>>> PropertyComparers;
 		TMap<FSoftClassPath, FCustomSerializer> CustomSerializers;
-		
+
+		TArray<TSharedRef<IActorSnapshotFilter>> GlobalFilters;
 		TArray<TSharedRef<ISnapshotLoader>> SnapshotLoaders;
 		TArray<TSharedRef<IRestorationListener>> RestorationListeners;
 
diff --git a/Engine/Plugins/VirtualProduction/LevelSnapshots/Source/LevelSnapshots/Private/Restorability/ActorGroupRestoration.cpp b/Engine/Plugins/VirtualProduction/LevelSnapshots/Source/LevelSnapshots/Private/Restorability/ActorGroupRestoration.cpp
--- a/Engine/Plugins/VirtualProduction/LevelSnapshots/Source/LevelSnapshots/Private/Restorability/ActorGroupRestoration.cpp
+++ b/Engine/Plugins/VirtualProduction/LevelSnapshots/Source/LevelSnapshots/Private/Restorability/ActorGroupRestoration.cpp
@@ -3,7 +3,9 @@
 #include "ActorGroupRestoration.h"
 
 #include "LevelSnapshot.h"
+#include "Algo/AllOf.h"
 #include "Editor/GroupActor.h"
+#include "GameFramework/WorldSettings.h"
 #include "Params/PropertyComparisonParams.h"
 #include "Selection/PropertySelectionMap.h"
 #include "Util/EquivalenceUtil.h"
@@ -22,10 +24,12 @@ namespace UE::LevelSnapshots::Private::ActorGroupRestoration
 		}
 	}
 	
-	class FActorGroupRestoration : public IRestorationListener, public IPropertyComparer
+	class FActorGroupRestoration : public IRestorationListener, public IPropertyComparer, public IActorSnapshotFilter
 	{
 		const FProperty* GroupActorsProperty = AGroupActor::StaticClass()->FindPropertyByName(GET_MEMBER_NAME_CHECKED(AGroupActor, GroupActors));
 		const FProperty* SubGroupsProperty = AGroupActor::StaticClass()->FindPropertyByName(GET_MEMBER_NAME_CHECKED(AGroupActor, SubGroups));
+
+		TSet<TWeakObjectPtr<AGroupActor>> ActorsWithDummyChild;
 		
 		template<typename TActor>
 		void RegisterActorsWithGroup(AGroupActor& GroupActor, TArray<TObjectPtr<TActor>>& Actors)
@@ -108,6 +112,12 @@ namespace UE::LevelSnapshots::Private::ActorGroupRestoration
 			return false;
 		}
 
+		AActor* GetHiddenDummyActorInWorld(UWorld& World)
+		{
+			// Should not be nullptr
+			return World.GetWorldSettings();
+		}
+
 	public:
 
 		FActorGroupRestoration()
@@ -138,16 +148,43 @@ namespace UE::LevelSnapshots::Private::ActorGroupRestoration
 
 			return EPropertyComparison::CheckNormally;
 		}
-
-		virtual void PreRemoveActor(AActor* ActorToRemove) override
+		
+		virtual void PreRemoveActor(const FPreRemoveActorParams& Params) override
 		{
-			if (AGroupActor* GroupActor = Cast<AGroupActor>(ActorToRemove))
+			if (AGroupActor* GroupActor = Cast<AGroupActor>(Params.ActorToRemove))
 			{
 				// This prevents the actors part of this group from also being deleted.
 				// This still allows users to include them in the deleted set of actors (that's good).
 				GroupActor->GroupActors.Reset();
 				GroupActor->SubGroups.Reset();
 			}
+			else if (AGroupActor* OwningActor = Cast<AGroupActor>(Params.ActorToRemove->GroupActor))
+			{
+				AActor* DummyActor = GetHiddenDummyActorInWorld(*OwningActor->GetWorld());
+				// Workflow: 1. Remove all actors in a group (e.g. they're all new) 2. Serialize new children into GroupActor
+				// The removal process would accidentally delete remove the group actor because it would no longer have any children
+				if (const FPropertySelection* Selection = Params.SelectedProperties.GetObjectSelection(OwningActor).GetPropertySelection()
+					; Selection && Selection->IsPropertySelected(nullptr, GroupActorsProperty) && DummyActor)
+				{
+					// This is pretty hacky but it will prevent AGroupActor::PostRemove from deleting the actor
+					OwningActor->GroupActors.Add(DummyActor);
+					ActorsWithDummyChild.Add(OwningActor);
+				}
+			}
+		}
+
+		virtual void PostRemoveActors(const FPostRemoveActorsParams& Params) override
+		{
+			// "Fix" the hack we did in PreRemoveActor
+			for (TWeakObjectPtr<AGroupActor> GroupActor : ActorsWithDummyChild)
+			{
+				if (GroupActor.IsValid())
+				{
+					AActor* DummyActor = GetHiddenDummyActorInWorld(*GroupActor->GetWorld());
+					GroupActor->GroupActors.Remove(DummyActor);
+				}
+			}
+			ActorsWithDummyChild.Reset();
 		}
 
 		virtual void PreApplySnapshotToActor(const FApplySnapshotToActorParams& Params) override
@@ -159,7 +196,7 @@ namespace UE::LevelSnapshots::Private::ActorGroupRestoration
 				// Clear every actor's AActor::GroupActor property if we're about to modify the GroupActor.
 				// The new actors will get reregistered in PostApplySnapshotToActor
 				// This is done so actors that will no longer be part of the group after the restore know they're no longer part of the group;
-				// otherwise you can click the actor in the World Outliner and it will highlight the old group
+				// otherwise you can click an actor that was removed in the World Outliner and it will highlight the old group
 				if (Selection.GetPropertySelection() && Selection.GetPropertySelection()->IsPropertySelected(nullptr, GroupActorsProperty))
 				{
 					UnregisterActorsFromGroup(*GroupActor, GroupActor->GroupActors);
@@ -180,8 +217,28 @@ namespace UE::LevelSnapshots::Private::ActorGroupRestoration
 				
 				RegisterActorsWithGroup(*GroupActor, GroupActor->GroupActors);
 				RegisterActorsWithGroup(*GroupActor, GroupActor->SubGroups);
+
+				// Hacky way of forcing GroupActor to reconsider whether it should exist...
+				// References can fail to restore ... possibly there are no children
+				GroupActor->Remove(*GroupActor);
 			}
 		}
+
+		virtual EFilterResult CanRecreateActor(const FCanRecreateActorParams& Params) override
+		{
+			if (Params.Class == AGroupActor::StaticClass())
+			{
+				if (const TOptional<TNonNullPtr<AActor>> SnapshotActor = Params.DeserializeFromSnapshotFunc())
+				{
+					const AGroupActor* AsGroupActor = Cast<AGroupActor>(SnapshotActor.GetValue());
+					check(AsGroupActor);
+					return Algo::AllOf(AsGroupActor->GroupActors, [](AActor* Actor){ return Actor != nullptr; })
+						? EFilterResult::DoNotCare : EFilterResult::Disallow;
+				}
+			}
+
+			return EFilterResult::DoNotCare;
+		}
 	};
 #endif
 	
@@ -191,6 +248,7 @@ namespace UE::LevelSnapshots::Private::ActorGroupRestoration
 		AddActorGroupSupport(Module);
 		const TSharedRef<FActorGroupRestoration> ActorGroupSupport = MakeShared<FActorGroupRestoration>();
 		Module.RegisterRestorationListener(ActorGroupSupport);
+		Module.RegisterGlobalActorFilter(ActorGroupSupport);
 		Module.RegisterPropertyComparer(AGroupActor::StaticClass(), ActorGroupSupport);
 #endif
 	}
diff --git a/Engine/Plugins/VirtualProduction/LevelSnapshots/Source/LevelSnapshots/Private/Restorability/SnapshotRestorability.cpp b/Engine/Plugins/VirtualProduction/LevelSnapshots/Source/LevelSnapshots/Private/Restorability/SnapshotRestorability.cpp
--- a/Engine/Plugins/VirtualProduction/LevelSnapshots/Source/LevelSnapshots/Private/Restorability/SnapshotRestorability.cpp
+++ b/Engine/Plugins/VirtualProduction/LevelSnapshots/Source/LevelSnapshots/Private/Restorability/SnapshotRestorability.cpp
@@ -2,13 +2,14 @@
 
 #include "Restorability/SnapshotRestorability.h"
 
+#include "Interfaces/IActorSnapshotFilter.h"
 #include "Landscape.h"
 #include "LandscapeGizmoActor.h"
 #include "LevelSnapshotsLog.h"
 #include "LevelSnapshotsModule.h"
+
 #include "AI/NavigationSystemConfig.h"
 #include "Algo/AllOf.h"
-
 #include "Components/ActorComponent.h"
 #include "Components/BillboardComponent.h"
 #include "Engine/BrushBuilder.h"
@@ -24,7 +25,7 @@
 
 namespace UE::LevelSnapshots::Restorability::Private::Internal
 {
-	static bool DoesActorHaveSupportedClass(const AActor* Actor)
+	static bool DoesActorHaveSupportedClass(const UClass* ActorClass)
 	{
 		const TSet<UClass*> UnsupportedClasses = 
 		{
@@ -36,7 +37,7 @@ namespace UE::LevelSnapshots::Restorability::Private::Internal
 			ALandscapeGizmoActor::StaticClass()
         };
 
-		return Algo::AllOf(UnsupportedClasses, [Actor](UClass *Class) { return !Actor->IsA(Class); });
+		return Algo::AllOf(UnsupportedClasses, [ActorClass](UClass* Class) { return !ActorClass->IsChildOf(Class); });
 	}
 
 	static bool DoesActorHaveSupportedClassForRemoving(const AActor* Actor)
@@ -92,7 +93,7 @@ bool UE::LevelSnapshots::Restorability::IsActorDesirableForCapture(const AActor*
 {
 	SCOPED_SNAPSHOT_CORE_TRACE(IsActorDesirableForCapture);
 
-	if (!IsValid(Actor) || !Private::Internal::DoesActorHaveSupportedClass(Actor))
+	if (!IsValid(Actor) || !Private::Internal::DoesActorHaveSupportedClass(Actor->GetClass()))
 	{
 		return false;
 	}
@@ -214,9 +215,16 @@ bool UE::LevelSnapshots::Restorability::IsPropertyExplicitlySupportedForCapture(
 	return UE::LevelSnapshots::Private::FLevelSnapshotsModule::GetInternalModuleInstance().IsPropertyExplicitlySupported(Property);
 }
 
+bool UE::LevelSnapshots::Restorability::ShouldConsiderRemovedActorForRecreation(const FCanRecreateActorParams& Params)
+{
+	return LevelSnapshots::Private::FLevelSnapshotsModule::GetInternalModuleInstance().CanRecreateActor(Params);
+}
+
 bool UE::LevelSnapshots::Restorability::ShouldConsiderNewActorForRemoval(const AActor* Actor)
 {
-	return Private::Internal::DoesActorHaveSupportedClassForRemoving(Actor) && IsActorDesirableForCapture(Actor);
+	return Private::Internal::DoesActorHaveSupportedClassForRemoving(Actor)
+		&& IsActorDesirableForCapture(Actor)
+		&& LevelSnapshots::Private::FLevelSnapshotsModule::GetInternalModuleInstance().CanDeleteActor(Actor);
 }
 
 bool UE::LevelSnapshots::Restorability::IsRestorableProperty(const FProperty* LeafProperty)
diff --git a/Engine/Plugins/VirtualProduction/LevelSnapshots/Source/LevelSnapshots/Public/Data/LevelSnapshot.h b/Engine/Plugins/VirtualProduction/LevelSnapshots/Source/LevelSnapshots/Public/Data/LevelSnapshot.h
--- a/Engine/Plugins/VirtualProduction/LevelSnapshots/Source/LevelSnapshots/Public/Data/LevelSnapshot.h
+++ b/Engine/Plugins/VirtualProduction/LevelSnapshots/Source/LevelSnapshots/Public/Data/LevelSnapshot.h
@@ -71,8 +71,8 @@ public:
 	 *	@param HandleAddedActor Actor exists in world but not in snapshot. Receives reference to world actor.
 	 */
 	UE_DEPRECATED(5.1, "Use DiffWorld that accepts FActorConsumer for HandleMatchedActor instead.")
-	void DiffWorld(UWorld* World, FActorPathConsumer HandleMatchedActor, FActorPathConsumer HandleRemovedActor, FActorConsumer HandleAddedActor) const;
-	void DiffWorld(UWorld* World, FActorConsumer HandleMatchedActor, FActorPathConsumer HandleRemovedActor, FActorConsumer HandleAddedActor) const;
+	void DiffWorld(UWorld* World, FActorPathConsumer HandleMatchedActor, FActorPathConsumer HandleRemovedActor, FActorConsumer HandleAddedActor);
+	void DiffWorld(UWorld* World, FActorConsumer HandleMatchedActor, FActorPathConsumer HandleRemovedActor, FActorConsumer HandleAddedActor);
 
 	
 	
diff --git a/Engine/Plugins/VirtualProduction/LevelSnapshots/Source/LevelSnapshots/Public/ILevelSnapshotsModule.h b/Engine/Plugins/VirtualProduction/LevelSnapshots/Source/LevelSnapshots/Public/ILevelSnapshotsModule.h
--- a/Engine/Plugins/VirtualProduction/LevelSnapshots/Source/LevelSnapshots/Public/ILevelSnapshotsModule.h
+++ b/Engine/Plugins/VirtualProduction/LevelSnapshots/Source/LevelSnapshots/Public/ILevelSnapshotsModule.h
@@ -8,6 +8,7 @@
 #include "Interfaces/IPropertyComparer.h"
 #include "Interfaces/ISnapshotRestorabilityOverrider.h"
 #include "Interfaces/ICustomObjectSnapshotSerializer.h"
+#include "Interfaces/IActorSnapshotFilter.h"
 
 #include "Modules/ModuleInterface.h"
 #include "Modules/ModuleManager.h"
@@ -54,26 +55,22 @@ namespace UE::LevelSnapshots
 
 		/** Queries the attached snapshot delegate and determines if we can take a snapshot.*/
 		virtual bool CanTakeSnapshot(const FPreTakeSnapshotEventData& Event) const = 0;
-
 		
 
 		/** Snapshots will no longer capture nor restore subobjects of this class. Subclasses are implicitly skipped as well. */
 		virtual void AddSkippedSubobjectClasses(const TSet<UClass*>& Classes) = 0;
 		virtual void RemoveSkippedSubobjectClasses(const TSet<UClass*>& Classes) = 0;
 		
-
 		
 		/* Registers callbacks that override which actors, components, and properties are restored by default. */
 		virtual void RegisterRestorabilityOverrider(TSharedRef<ISnapshotRestorabilityOverrider> Overrider) = 0;
 		/* Unregisters an overrider previously registered. */
-		virtual void UnregisterRestorabilityOverrider(TSharedRef<ISnapshotRestorabilityOverrider> Overrider) = 0;
-
+		virtual void UnregisterRestorabilityOverrider(const TSharedRef<ISnapshotRestorabilityOverrider>& Overrider) = 0;
 		
 		
 		/* Registers a callback for deciding whether a property should be considered changed. Applies to all sub-classes. */
 		virtual void RegisterPropertyComparer(UClass* Class, TSharedRef<IPropertyComparer> Comparer) = 0;
-		virtual void UnregisterPropertyComparer(UClass* Class, TSharedRef<IPropertyComparer> Comparer) = 0;
-
+		virtual void UnregisterPropertyComparer(UClass* Class, const TSharedRef<IPropertyComparer>& Comparer) = 0;
 		
 		
 		/**
@@ -86,16 +83,21 @@ namespace UE::LevelSnapshots
 		 */
 		virtual void RegisterCustomObjectSerializer(UClass* Class, TSharedRef<ICustomObjectSnapshotSerializer> CustomSerializer, bool bIncludeBlueprintChildClasses = true) = 0;
 		virtual void UnregisterCustomObjectSerializer(UClass* Class) = 0;
+		
+		
+		/** Registers callbacks for external systems to decide whether a given actor can be modified, recreated, or removed from the world. */
+		virtual void RegisterGlobalActorFilter(TSharedRef<IActorSnapshotFilter> Filter) = 0;
+		virtual void UnregisterGlobalActorFilter(const TSharedRef<IActorSnapshotFilter>& Filter) = 0;
 
 
 		/** Register an object that will receive callbacks when a snapshot is loaded */
 		virtual void RegisterSnapshotLoader(TSharedRef<ISnapshotLoader> Loader) = 0;
-		virtual void UnregisterSnapshotLoader(TSharedRef<ISnapshotLoader> Loader) = 0;
+		virtual void UnregisterSnapshotLoader(const TSharedRef<ISnapshotLoader>& Loader) = 0;
 		
 		
 		/** Registers an object that will receive callbacks when a snapshot is applied. */
 		virtual void RegisterRestorationListener(TSharedRef<IRestorationListener> Listener) = 0;
-		virtual void UnregisterRestorationListener(TSharedRef<IRestorationListener> Listener) = 0;
+		virtual void UnregisterRestorationListener(const TSharedRef<IRestorationListener>& Listener) = 0;
 		
 		
 		/**
diff --git a/Engine/Plugins/VirtualProduction/LevelSnapshots/Source/LevelSnapshots/Public/Restorability/Interfaces/IActorSnapshotFilter.h b/Engine/Plugins/VirtualProduction/LevelSnapshots/Source/LevelSnapshots/Public/Restorability/Interfaces/IActorSnapshotFilter.h
new file mode 100644
--- /dev/null
+++ b/Engine/Plugins/VirtualProduction/LevelSnapshots/Source/LevelSnapshots/Public/Restorability/Interfaces/IActorSnapshotFilter.h
@@ -0,0 +1,53 @@
+// Copyright Epic Games, Inc. All Rights Reserved.
+
+#pragma once
+
+#include "CoreMinimal.h"
+#include "Templates/NonNullPointer.h"
+#include "UObject/SoftObjectPath.h"
+
+class AActor;
+class UObject;
+
+namespace UE::LevelSnapshots
+{
+	using FDeserializedActorGetter = TFunctionRef<TOptional<TNonNullPtr<AActor>>()>;
+	
+	struct FCanRecreateActorParams
+	{
+		UWorld* EditorWorld;
+
+		UClass* Class;
+		FSoftObjectPath ActorPath;
+		
+		/** If you need access to the actor's data, you can call this function. */
+		FDeserializedActorGetter DeserializeFromSnapshotFunc;
+	};
+	
+	/**
+	 * Can decide whether a given actor can be added to or removed from the world.
+	 */
+	class LEVELSNAPSHOTS_API IActorSnapshotFilter
+	{
+	public:
+
+		enum class EFilterResult
+		{
+			/* The object in question is included but only if nobody else returned Disallow.  */
+			Allow,
+			/* Let other filters decide. If every override return DoNotCare, default snapshot behaviour is applied. */
+			DoNotCare,
+			/* The object in question is never suitable and is not included. Other filters cannot override this. */
+			Disallow
+		};
+
+		/** Whether the actor is allowed to be recreated. Called when diffing to world. */
+		virtual EFilterResult CanRecreateActor(const FCanRecreateActorParams& Params) { return EFilterResult::DoNotCare; }
+
+		/** Whether the actor is allowed to be removed. Called when diffing to world. */
+		virtual EFilterResult CanDeleteActor(const AActor* EditorActor) { return EFilterResult::DoNotCare; }
+
+		
+		virtual ~IActorSnapshotFilter() = default;
+	};
+}
\ No newline at end of file
diff --git a/Engine/Plugins/VirtualProduction/LevelSnapshots/Source/LevelSnapshots/Public/Restorability/Interfaces/IRestorationListener.h b/Engine/Plugins/VirtualProduction/LevelSnapshots/Source/LevelSnapshots/Public/Restorability/Interfaces/IRestorationListener.h
--- a/Engine/Plugins/VirtualProduction/LevelSnapshots/Source/LevelSnapshots/Public/Restorability/Interfaces/IRestorationListener.h
+++ b/Engine/Plugins/VirtualProduction/LevelSnapshots/Source/LevelSnapshots/Public/Restorability/Interfaces/IRestorationListener.h
@@ -15,6 +15,29 @@ struct FPropertySelectionMap;
 
 namespace UE::LevelSnapshots
 {
+	struct FApplySnapshotParams
+	{
+		/** All of the user's selected properties */
+		const FPropertySelectionMap& SelectedProperties;
+
+		FApplySnapshotParams(const FPropertySelectionMap& SelectedProperties)
+			: SelectedProperties(SelectedProperties)
+		{}
+	};
+
+	struct FPreRemoveActorParams : FApplySnapshotParams
+	{
+		AActor* ActorToRemove;
+
+		FPreRemoveActorParams(const FPropertySelectionMap& SelectedProperties, AActor* ActorToRemove)
+			: FApplySnapshotParams(SelectedProperties),
+			  ActorToRemove(ActorToRemove)
+		{}
+	};
+
+	struct FPostRemoveActorsParams : FApplySnapshotParams
+	{};
+	
 	struct FApplySnapshotPropertiesParams
 	{
 		/** The object that receives serialized data. */
@@ -107,6 +130,14 @@ namespace UE::LevelSnapshots
 	class LEVELSNAPSHOTS_API IRestorationListener
 	{
 	public:
+
+		/** Called before the snapshot is applied. Except for the changes made by this event, the world has not been changed, yet. */
+		virtual void PreApplySnapshot(const FApplySnapshotParams& Params) {}
+		
+		/** Called after all snapshot data was applied. The world is nearly completely changed (the only thing remaining is executing this event). */
+		virtual void PostApplySnapshot(const FApplySnapshotParams& Params) {}
+
+
 		
 		/**
 		 * Called before applying snapshot data to an object.
@@ -141,6 +172,7 @@ namespace UE::LevelSnapshots
 		 */
 		virtual void PostApplySnapshotToActor(const FApplySnapshotToActorParams& Params) {}
 
+		
 
 		/**
 		 * Called before an actor is recreated to the world and gives the opportunity to override settings.
@@ -150,13 +182,18 @@ namespace UE::LevelSnapshots
 		
 		/** Called after an actor is recreated to the world (but before any data is applied to it) */
 		virtual void PostRecreateActor(AActor* RecreatedActor) {}
+		
 
-
+		UE_DEPRECATED(5.1, "Use PreRemoveActor(const FPreRemoveActorParams& Params) instead")
+		virtual void PreRemoveActor(AActor* ActorToRemove) {}
 		
 		/** Called before an actor is removed from the world. */
-		virtual void PreRemoveActor(AActor* ActorToRemove) {}
-
+		virtual void PreRemoveActor(const FPreRemoveActorParams& Params) {}
+		
+		/** Called after all actors have been removed from the world */
+		virtual void PostRemoveActors(const FPostRemoveActorsParams& Params) {}
 		
+
 		
 		/**
 		 * Called before a component is recreated on an actor.
diff --git a/Engine/Plugins/VirtualProduction/LevelSnapshots/Source/LevelSnapshots/Public/Restorability/SnapshotRestorability.h b/Engine/Plugins/VirtualProduction/LevelSnapshots/Source/LevelSnapshots/Public/Restorability/SnapshotRestorability.h
--- a/Engine/Plugins/VirtualProduction/LevelSnapshots/Source/LevelSnapshots/Public/Restorability/SnapshotRestorability.h
+++ b/Engine/Plugins/VirtualProduction/LevelSnapshots/Source/LevelSnapshots/Public/Restorability/SnapshotRestorability.h
@@ -4,10 +4,19 @@
 
 #include "CoreMinimal.h"
 
+
 class AActor;
 class FProperty;
 class UActorComponent;
 
+namespace UE
+{
+	namespace LevelSnapshots
+	{
+		struct FCanRecreateActorParams;
+	}
+}
+
 namespace UE::LevelSnapshots::Restorability
 {
 	/* Is this actor captured by the snapshot system? */
@@ -30,6 +39,8 @@ namespace UE::LevelSnapshots::Restorability
 	/* Is this property always captured by the snapshot system? */
 	LEVELSNAPSHOTS_API bool IsPropertyExplicitlySupportedForCapture(const FProperty* Property);
 
+	/** The actor actors was removed from the snapshot. Should we show it in the list of removed actors? */
+	LEVELSNAPSHOTS_API bool ShouldConsiderRemovedActorForRecreation(const FCanRecreateActorParams& Params);
 	/* The actor did not exist in the snapshot. Should we show it in the list of added actors? */
 	LEVELSNAPSHOTS_API bool ShouldConsiderNewActorForRemoval(const AActor* Actor);
 	
