--------
Fix crash on retaking snapshot on existing asset

#[fyi] Jason.Walter
Fix unit tests and fix crash when retaking a snapshot

#[fyi] Jason.Walter

diff --git a/Engine/Plugins/VirtualProduction/LevelSnapshots/Source/LevelSnapshots/Private/Data/LevelSnapshot.cpp b/Engine/Plugins/VirtualProduction/LevelSnapshots/Source/LevelSnapshots/Private/Data/LevelSnapshot.cpp
--- a/Engine/Plugins/VirtualProduction/LevelSnapshots/Source/LevelSnapshots/Private/Data/LevelSnapshot.cpp
+++ b/Engine/Plugins/VirtualProduction/LevelSnapshots/Source/LevelSnapshots/Private/Data/LevelSnapshot.cpp
@@ -19,6 +19,7 @@
 #include "Algo/Accumulate.h"
 #include "Engine/Engine.h"
 #include "EngineUtils.h"
+#include "Algo/IndexOf.h"
 #include "HAL/IConsoleManager.h"
 #include "GameFramework/Actor.h"
 #include "Misc/ScopeExit.h"
@@ -82,9 +83,9 @@ bool ULevelSnapshot::SnapshotWorld(UWorld* TargetWorld)
 
 	MapPath = TargetWorld;
 	CaptureTime = FDateTime::UtcNow();
+	SerializedData = UE::LevelSnapshots::Private::SnapshotWorld(TargetWorld);
 	// Existing snapshot world actors are now invalid - clear previous cached data
 	RecreateSnapshotWorld();
-	SerializedData = UE::LevelSnapshots::Private::SnapshotWorld(TargetWorld);
 	ClearCache();
 	Cache.InitFor(SerializedData);
 	// Should we delete preexisting actors from the snapshot world?
@@ -326,22 +327,12 @@ FString ULevelSnapshot::GenerateDebugLogInfo() const
 	return FString::Printf(TEXT("CaptureTime: %s. SnapshotVersionInfo: %s. Current engine version: %s."), *CaptureTime.ToString(), *GetSerializedData().SnapshotVersionInfo.ToString(), *Current.ToString());
 }
 
-void ULevelSnapshot::RecreateSnapshotWorld()
-{
-	if (RootSnapshotWorld)
-	{
-		DestroyWorld();
-	}
-	EnsureWorldInitialised();
-}
-
 namespace UE::LevelSnapshots::Private
 {
 	static UWorld* CreateWorld(UObject* Outer, FName WorldName);
 	static TSet<FName> ExtractLevelNames(const FWorldSnapshotData& WorldData);
 }
 
-
 void ULevelSnapshot::EnsureWorldInitialised()
 {
 	using namespace UE::LevelSnapshots::Private;
@@ -386,13 +377,8 @@ void ULevelSnapshot::EnsureWorldInitialised()
 
 namespace UE::LevelSnapshots::Private
 {
-	static UWorld* CreateWorld(UObject* Outer, FName WorldName)
+	static void InitializeWorld(UWorld* World)
 	{
-		UWorld* World = NewObject<UWorld>(Outer, WorldName);
-		World->WorldType = EWorldType::EditorPreview;
-
-		// Note: Do NOT create a FWorldContext for this world.
-		// If you do, the render thread will send render commands every tick (and crash cuz we do not init the scene below).
 		World->InitializeNewWorld(UWorld::InitializationValues()
 			.InitializeScenes(false)		// This is memory only world: no rendering
 			.AllowAudioPlayback(false)
@@ -405,7 +391,16 @@ namespace UE::LevelSnapshots::Private
 			.SetTransactional(false)
 			.CreateFXSystem(false)
 			);
+	}
+	
+	static UWorld* CreateWorld(UObject* Outer, FName WorldName)
+	{
+		UWorld* World = NewObject<UWorld>(Outer, WorldName);
+		World->WorldType = EWorldType::EditorPreview;
 
+		// Note: Do NOT create a FWorldContext for this world.
+		// If you do, the render thread will send render commands every tick (and crash cuz we do not init the scene below).
+		InitializeWorld(World);
 		return World;
 	}
 
@@ -430,30 +425,100 @@ void ULevelSnapshot::DestroyWorld()
 			GEngine->OnWorldDestroyed().Remove(Handle);
 			Handle.Reset();
 		}
-				
-		SerializedData.SnapshotWorld.Reset();
-		SerializedData.SnapshotSublevels.Reset();
-		ClearCache();
-	
-		RootSnapshotWorld->CleanupWorld();
-		for (UWorld* Sublevel : SnapshotSublevels)
-		{
-			Sublevel->CleanupWorld();
-		}
-		// Cleanup world clears RF_Standalone flag...
-		SetFlags(RF_Standalone);
 		
+		CleanUpWorld();
 		RootSnapshotWorld = nullptr;
 		SnapshotSublevels.Reset();
 	}
 }
 
+void ULevelSnapshot::CleanUpWorld()
+{
+	SerializedData.SnapshotWorld.Reset();
+	SerializedData.SnapshotSublevels.Reset();
+	ClearCache();
+
+	// Cleanup world clears RF_Standalone flag...
+	const bool bWasStandalone = HasAnyFlags(RF_Standalone);
+	RootSnapshotWorld->CleanupWorld();
+	for (UWorld* Sublevel : SnapshotSublevels)
+	{
+		Sublevel->CleanupWorld();
+	}
+	
+	if (bWasStandalone)
+	{
+		SetFlags(RF_Standalone);
+	}
+}
+
 void ULevelSnapshot::ClearCache() 
 {
 	Cache.Reset();
 	CachedDiffedActors.Reset(); 
 }
 
+void ULevelSnapshot::RecreateSnapshotWorld()
+{
+	if (RootSnapshotWorld)
+	{
+		DestroyWorld();
+
+		if (HasAnyFlags(RF_Standalone))
+		{
+			CollectGarbage(RF_Standalone);
+		}
+		else
+		{
+			SetFlags(RF_Standalone);
+			CollectGarbage(RF_Standalone);
+			ClearFlags(RF_Standalone);
+		}
+	}
+	else
+	{
+		EnsureWorldInitialised();
+	}
+}
+
+namespace UE::LevelSnapshots::Private
+{
+	UWorld* FindOrCreateWorld(const TArray<UWorld*>& AllWorlds, UObject* Outer, FName Name);
+}
+
+void ULevelSnapshot::ReinitializeWorld()
+{
+	using namespace UE::LevelSnapshots::Private;
+	
+	TArray<UWorld*> AllWorlds = MoveTemp(SnapshotSublevels);
+	AllWorlds.Add(RootSnapshotWorld);
+
+	const FName RootWorldName = *ExtractPathWithoutSubobjects(MapPath).GetAssetName();
+	RootSnapshotWorld = FindOrCreateWorld(AllWorlds, this, RootWorldName);
+		
+	TSet<FName> Sublevels = ExtractLevelNames(SerializedData);
+	Sublevels.Remove(RootWorldName);
+	SnapshotSublevels.Reserve(Sublevels.Num());
+	for (const FName& SublevelName : Sublevels)
+	{
+		SnapshotSublevels.Add(FindOrCreateWorld(AllWorlds, this, SublevelName));
+	}
+}
+
+namespace UE::LevelSnapshots::Private
+{
+	UWorld* FindOrCreateWorld(const TArray<UWorld*>& AllWorlds, UObject* Outer, FName Name)
+	{
+		const int32 Index = Algo::IndexOfByPredicate(AllWorlds, [Name](UWorld* World){ return World->GetFName() == Name; });
+		if (AllWorlds.IsValidIndex(Index))
+		{
+			InitializeWorld(AllWorlds[Index]);
+			return AllWorlds[Index];
+		}
+		return CreateWorld(Outer, Name);
+	}
+}
+
 #if WITH_EDITOR
 void ULevelSnapshot::ResetDiffCacheToUninitialized(TArrayView<AActor*> ModifiedActors)
 {
diff --git a/Engine/Plugins/VirtualProduction/LevelSnapshots/Source/LevelSnapshots/Private/Tests/CustomObjectSerializationTests.cpp b/Engine/Plugins/VirtualProduction/LevelSnapshots/Source/LevelSnapshots/Private/Tests/CustomObjectSerializationTests.cpp
--- a/Engine/Plugins/VirtualProduction/LevelSnapshots/Source/LevelSnapshots/Private/Tests/CustomObjectSerializationTests.cpp
+++ b/Engine/Plugins/VirtualProduction/LevelSnapshots/Source/LevelSnapshots/Private/Tests/CustomObjectSerializationTests.cpp
@@ -231,11 +231,12 @@ namespace UE::LevelSnapshots::Private::Tests
 			virtual void OnTakeSnapshot(UObject* EditorObject, ICustomSnapshotSerializationData& DataStorage) override
 			{
 				ASnapshotTestActor* Actor = Cast<ASnapshotTestActor>(EditorObject);
-				const int32 Index = DataStorage.AddSubobjectSnapshot(Actor->NonReflectedSubobject);
+				const int32 Index = DataStorage.AddSubobjectSnapshot(Actor->NonReflectedSubobject.Get());
 				DataStorage.WriteObjectAnnotation(FObjectAnnotator::CreateLambda([Actor](FArchive& Archive)
 				{
+					UObject* Object = Actor->NonReflectedObjectProperty.Get(); 
 					Archive << Actor->NonReflectedName;
-					Archive << Actor->NonReflectedObjectProperty;
+					Archive << Object;
 					Archive << Actor->NonReflectedSoftPtr;
 				}));
 				DataStorage.GetSubobjectMetaData(Index)->WriteObjectAnnotation(FObjectAnnotator::CreateLambda([Actor](FArchive& Archive)
@@ -250,20 +251,20 @@ namespace UE::LevelSnapshots::Private::Tests
 			{
 				ASnapshotTestActor* Actor = Cast<ASnapshotTestActor>(SnapshotObject);
 				Actor->AllocateNonReflectedSubobject();
-				return Actor->NonReflectedSubobject;
+				return Actor->NonReflectedSubobject.Get();
 			}
 		
 			virtual UObject* FindOrRecreateSubobjectInEditorWorld(UObject* EditorObject, const ISnapshotSubobjectMetaData& ObjectData, const ICustomSnapshotSerializationData& DataStorage) override
 			{
 				ASnapshotTestActor* Actor = Cast<ASnapshotTestActor>(EditorObject);
 				Actor->AllocateNonReflectedSubobject();
-				return Actor->NonReflectedSubobject;
+				return Actor->NonReflectedSubobject.Get();
 			}
 		
 			virtual UObject* FindSubobjectInEditorWorld(UObject* EditorObject, const ISnapshotSubobjectMetaData& ObjectData, const ICustomSnapshotSerializationData& DataStorage) override
 			{
 				ASnapshotTestActor* Actor = Cast<ASnapshotTestActor>(EditorObject);
-				return Actor->NonReflectedSubobject;
+				return Actor->NonReflectedSubobject.Get();
 			}
 
 			virtual void PostApplySnapshotProperties(UObject* OriginalObject, const ICustomSnapshotSerializationData& DataStorage) override
@@ -271,9 +272,11 @@ namespace UE::LevelSnapshots::Private::Tests
 				ASnapshotTestActor* Actor = Cast<ASnapshotTestActor>(OriginalObject);
 				DataStorage.ReadObjectAnnotation(FObjectAnnotator::CreateLambda([Actor](FArchive& Archive)
 				{
+					UObject* Object = nullptr; 
 					Archive << Actor->NonReflectedName;
-					Archive << Actor->NonReflectedObjectProperty;
+					Archive << Object;
 					Archive << Actor->NonReflectedSoftPtr;
+					Actor->NonReflectedObjectProperty.Reset(Object);
 				}));
 
 				DataStorage.GetSubobjectMetaData(0)->ReadObjectAnnotation(FObjectAnnotator::CreateLambda([Actor](FArchive& Archive)
@@ -312,7 +315,7 @@ namespace UE::LevelSnapshots::Private::Tests
 			
 				TestActor->IntProperty = 1;
 				TestActor->NonReflectedName = FName("TestNonReflectedName_OnActor");
-				TestActor->NonReflectedObjectProperty = FirstReferencedActor;
+				TestActor->NonReflectedObjectProperty.Reset(FirstReferencedActor);
 				TestActor->NonReflectedSoftPtr = FirstReferencedActor;
 			
 				TestActor->NonReflectedSubobject->NonReflectedName = FName("TestNonReflectedName_OnSubobject");
@@ -339,7 +342,7 @@ namespace UE::LevelSnapshots::Private::Tests
 			.RunTest([&]()
 			{
 				TestEqual(TEXT("TestActor->NonReflectedName"), TestActor->NonReflectedName, FName("TestNonReflectedName_OnActor"));
-				TestTrue(TEXT("TestActor->NonReflectedObjectProperty"), TestActor->NonReflectedObjectProperty == FirstReferencedActor);
+				TestTrue(TEXT("TestActor->NonReflectedObjectProperty"), TestActor->NonReflectedObjectProperty.Get() == FirstReferencedActor);
 				TestTrue(TEXT("TestActor->NonReflectedSoftPtr"), TestActor->NonReflectedSoftPtr == FirstReferencedActor);
 			
 				TestEqual(TEXT("TestActor->NonReflectedSubobject->NonReflectedName"), TestActor->NonReflectedSubobject->NonReflectedName, FName("TestNonReflectedName_OnSubobject"));
@@ -726,7 +729,7 @@ namespace UE::LevelSnapshots::Private::Tests
 			{
 				if (TestActor->NonReflectedSubobject)
 				{
-					DataStorage.AddSubobjectSnapshot(TestActor->NonReflectedSubobject);
+					DataStorage.AddSubobjectSnapshot(TestActor->NonReflectedSubobject.Get());
 				}
 			}
 
@@ -735,7 +738,7 @@ namespace UE::LevelSnapshots::Private::Tests
 				if (ASnapshotTestActor* SnapshotActor = Cast<ASnapshotTestActor>(SnapshotObject))
 				{
 					SnapshotActor->AllocateNonReflectedSubobject();
-					return SnapshotActor->NonReflectedSubobject;
+					return SnapshotActor->NonReflectedSubobject.Get();
 				}
 
 				checkNoEntry();
@@ -747,7 +750,7 @@ namespace UE::LevelSnapshots::Private::Tests
 				if (ASnapshotTestActor* EditorActor = Cast<ASnapshotTestActor>(EditorObject))
 				{
 					EditorActor->AllocateNonReflectedSubobject();
-					return EditorActor->NonReflectedSubobject;
+					return EditorActor->NonReflectedSubobject.Get();
 				}
 
 				checkNoEntry();
@@ -758,7 +761,7 @@ namespace UE::LevelSnapshots::Private::Tests
 			{
 				if (ASnapshotTestActor* EditorActor = Cast<ASnapshotTestActor>(EditorObject))
 				{
-					return EditorActor->NonReflectedSubobject;
+					return EditorActor->NonReflectedSubobject.Get();
 				}
 
 				checkNoEntry();
@@ -808,7 +811,7 @@ namespace UE::LevelSnapshots::Private::Tests
 
 			.RunTest([&]()
 			{
-				if (IsValid(Stub->TestActor->NonReflectedSubobject))
+				if (IsValid(Stub->TestActor->NonReflectedSubobject.Get()))
 				{
 					TestEqual(TEXT("Custom float restored"), Stub->TestActor->NonReflectedSubobject->FloatProperty, 42.f);
 					TestEqual(TEXT("Custom float restored"), Stub->TestActor->NonReflectedSubobject->IntProperty, 21);
diff --git a/Engine/Plugins/VirtualProduction/LevelSnapshots/Source/LevelSnapshots/Private/Tests/Types/SnapshotTestActor.cpp b/Engine/Plugins/VirtualProduction/LevelSnapshots/Source/LevelSnapshots/Private/Tests/Types/SnapshotTestActor.cpp
--- a/Engine/Plugins/VirtualProduction/LevelSnapshots/Source/LevelSnapshots/Private/Tests/Types/SnapshotTestActor.cpp
+++ b/Engine/Plugins/VirtualProduction/LevelSnapshots/Source/LevelSnapshots/Private/Tests/Types/SnapshotTestActor.cpp
@@ -171,17 +171,13 @@ void ASnapshotTestActor::AllocateNonReflectedSubobject()
 {
 	if (!NonReflectedSubobject)
 	{
-		NonReflectedSubobject = NewObject<USubobject>(this, USubobject::StaticClass(), TEXT("NonReflectedSubobject"));
+		NonReflectedSubobject.Reset(NewObject<USubobject>(this, USubobject::StaticClass(), TEXT("NonReflectedSubobject")));
 	}
 }
 
 void ASnapshotTestActor::DestroyNonReflectedSubobject()
 {
-	if (NonReflectedSubobject)
-	{
-		NonReflectedSubobject->MarkAsGarbage();
-		NonReflectedSubobject = nullptr;
-	}
+	NonReflectedSubobject.Reset();
 }
 
 void ASnapshotTestActor::PostInitProperties()
diff --git a/Engine/Plugins/VirtualProduction/LevelSnapshots/Source/LevelSnapshots/Private/Tests/Types/SnapshotTestActor.h b/Engine/Plugins/VirtualProduction/LevelSnapshots/Source/LevelSnapshots/Private/Tests/Types/SnapshotTestActor.h
--- a/Engine/Plugins/VirtualProduction/LevelSnapshots/Source/LevelSnapshots/Private/Tests/Types/SnapshotTestActor.h
+++ b/Engine/Plugins/VirtualProduction/LevelSnapshots/Source/LevelSnapshots/Private/Tests/Types/SnapshotTestActor.h
@@ -244,9 +244,9 @@ public:
 	UPROPERTY(EditAnywhere, Category = "Level Snapshots")
 	TMap<FName, USubobject*> EditOnlySubobjectMap_OptionalSubobject;
 	
-	USubobject* NonReflectedSubobject;
+	TStrongObjectPtr<USubobject> NonReflectedSubobject;
 
 	FName NonReflectedName;
-	UObject* NonReflectedObjectProperty;
+	TStrongObjectPtr<UObject> NonReflectedObjectProperty;
 	TSoftObjectPtr<UObject> NonReflectedSoftPtr;
 };
diff --git a/Engine/Plugins/VirtualProduction/LevelSnapshots/Source/LevelSnapshots/Public/Data/LevelSnapshot.h b/Engine/Plugins/VirtualProduction/LevelSnapshots/Source/LevelSnapshots/Public/Data/LevelSnapshot.h
--- a/Engine/Plugins/VirtualProduction/LevelSnapshots/Source/LevelSnapshots/Public/Data/LevelSnapshot.h
+++ b/Engine/Plugins/VirtualProduction/LevelSnapshots/Source/LevelSnapshots/Public/Data/LevelSnapshot.h
@@ -117,11 +117,14 @@ private:
 
 	FString GenerateDebugLogInfo() const;
 
-	void RecreateSnapshotWorld();
 	void EnsureWorldInitialised();
 	void DestroyWorld();
+	void CleanUpWorld();
 	void ClearCache();
 	
+	void RecreateSnapshotWorld();
+	void ReinitializeWorld();
+	
 	/** Callback to destroy our world when editor (editor build) or play (game builds) world is destroyed. */
 	FDelegateHandle Handle;
 	/** Callback to when an object is modified. Clears FActorSnapshotData::bHasBeenDiffed */
