Fixed non-unity build

#rnx

diff --git a/Engine/Plugins/Animation/ControlRig/Source/ControlRig/Private/Rigs/RigHierarchy.cpp b/Engine/Plugins/Animation/ControlRig/Source/ControlRig/Private/Rigs/RigHierarchy.cpp
--- a/Engine/Plugins/Animation/ControlRig/Source/ControlRig/Private/Rigs/RigHierarchy.cpp
+++ b/Engine/Plugins/Animation/ControlRig/Source/ControlRig/Private/Rigs/RigHierarchy.cpp
@@ -21,6 +21,7 @@ LLM_DEFINE_TAG(Animation_ControlRig);
 #include "HAL/PlatformFileManager.h"
 #include "Misc/FileHelper.h"
 #include "Misc/Paths.h"
+#include "Misc/TransactionObjectEvent.h"
 #include "JsonObjectConverter.h"
 #include "Serialization/JsonReader.h"
 #include "Serialization/JsonSerializer.h"
diff --git a/Engine/Plugins/Experimental/PCG/Source/PCGEditor/Private/PCGEditorGraphNodeBase.cpp b/Engine/Plugins/Experimental/PCG/Source/PCGEditor/Private/PCGEditorGraphNodeBase.cpp
--- a/Engine/Plugins/Experimental/PCG/Source/PCGEditor/Private/PCGEditorGraphNodeBase.cpp
+++ b/Engine/Plugins/Experimental/PCG/Source/PCGEditor/Private/PCGEditorGraphNodeBase.cpp
@@ -14,6 +14,7 @@
 #include "EdGraph/EdGraph.h"
 #include "EdGraph/EdGraphPin.h"
 #include "Framework/Commands/GenericCommands.h"
+#include "Misc/TransactionObjectEvent.h"
 #include "GraphEditorActions.h"
 #include "GraphEditorSettings.h"
 #include "ToolMenu.h"
diff --git a/Engine/Plugins/Importers/USDImporter/Source/USDStage/Public/USDStageActor.h b/Engine/Plugins/Importers/USDImporter/Source/USDStage/Public/USDStageActor.h
--- a/Engine/Plugins/Importers/USDImporter/Source/USDStage/Public/USDStageActor.h
+++ b/Engine/Plugins/Importers/USDImporter/Source/USDStage/Public/USDStageActor.h
@@ -4,6 +4,7 @@
 
 #include "GameFramework/Actor.h"
 #include "Misc/SecureHash.h"
+#include "Misc/ITransaction.h"
 #include "Templates/Function.h"
 
 #include "USDInfoCache.h"
diff --git a/Engine/Plugins/VirtualProduction/LevelSnapshots/Source/LevelSnapshotsEditor/Private/Data/FilterLoader.cpp b/Engine/Plugins/VirtualProduction/LevelSnapshots/Source/LevelSnapshotsEditor/Private/Data/FilterLoader.cpp
--- a/Engine/Plugins/VirtualProduction/LevelSnapshots/Source/LevelSnapshotsEditor/Private/Data/FilterLoader.cpp
+++ b/Engine/Plugins/VirtualProduction/LevelSnapshots/Source/LevelSnapshotsEditor/Private/Data/FilterLoader.cpp
@@ -8,7 +8,7 @@
 #include "EditorDirectories.h"
 #include "FileHelpers.h"
 #include "Misc/Paths.h"
-#include "Misc/ITransaction.h"
+#include "Misc/TransactionObjectEvent.h"
 #include "ScopedTransaction.h"
 
 void UFilterLoader::OverwriteExisting()
diff --git a/Engine/Plugins/VirtualProduction/LevelSnapshots/Source/LevelSnapshotsEditor/Private/Data/Filters/LevelSnapshotsFilterPreset.cpp b/Engine/Plugins/VirtualProduction/LevelSnapshots/Source/LevelSnapshotsEditor/Private/Data/Filters/LevelSnapshotsFilterPreset.cpp
--- a/Engine/Plugins/VirtualProduction/LevelSnapshots/Source/LevelSnapshotsEditor/Private/Data/Filters/LevelSnapshotsFilterPreset.cpp
+++ b/Engine/Plugins/VirtualProduction/LevelSnapshots/Source/LevelSnapshotsEditor/Private/Data/Filters/LevelSnapshotsFilterPreset.cpp
@@ -4,7 +4,7 @@
 
 #include "Data/Filters/ConjunctionFilter.h"
 
-#include "Misc/ITransaction.h"
+#include "Misc/TransactionObjectEvent.h"
 #include "ScopedTransaction.h"
 #include "UObject/UObjectGlobals.h"
 
diff --git a/Engine/Plugins/VirtualProduction/Takes/Source/TakesCore/Private/TakePreset.cpp b/Engine/Plugins/VirtualProduction/Takes/Source/TakesCore/Private/TakePreset.cpp
--- a/Engine/Plugins/VirtualProduction/Takes/Source/TakesCore/Private/TakePreset.cpp
+++ b/Engine/Plugins/VirtualProduction/Takes/Source/TakesCore/Private/TakePreset.cpp
@@ -4,6 +4,7 @@
 #include "LevelSequence.h"
 #include "MovieScene.h"
 #include "MovieSceneToolsProjectSettings.h"
+#include "Misc/TransactionObjectEvent.h"
 #include "UObject/Package.h"
 
 UTakePreset::UTakePreset(const FObjectInitializer& ObjInit)
diff --git a/Engine/Source/Developer/RigVMDeveloper/Private/RigVMModel/RigVMClient.cpp b/Engine/Source/Developer/RigVMDeveloper/Private/RigVMModel/RigVMClient.cpp
--- a/Engine/Source/Developer/RigVMDeveloper/Private/RigVMModel/RigVMClient.cpp
+++ b/Engine/Source/Developer/RigVMDeveloper/Private/RigVMModel/RigVMClient.cpp
@@ -1,7 +1,7 @@
 // Copyright Epic Games, Inc. All Rights Reserved.
 
 #include "RigVMModel/RigVMClient.h"
-#include "Misc/ITransaction.h"
+#include "Misc/TransactionObjectEvent.h"
 #if WITH_EDITOR
 #include "ScopedTransaction.h"
 #endif
diff --git a/Engine/Source/Developer/RigVMDeveloper/Private/RigVMModel/RigVMControllerActions.cpp b/Engine/Source/Developer/RigVMDeveloper/Private/RigVMModel/RigVMControllerActions.cpp
--- a/Engine/Source/Developer/RigVMDeveloper/Private/RigVMModel/RigVMControllerActions.cpp
+++ b/Engine/Source/Developer/RigVMDeveloper/Private/RigVMModel/RigVMControllerActions.cpp
@@ -4,7 +4,7 @@
 #include "UObject/Package.h"
 
 #if WITH_EDITOR
-#include "Misc/ITransaction.h"
+#include "Misc/TransactionObjectEvent.h"
 #endif
 
 UScriptStruct* FRigVMActionKey::GetScriptStruct() const
diff --git a/Engine/Source/Editor/UnrealEd/Public/EditorUndoClient.h b/Engine/Source/Editor/UnrealEd/Public/EditorUndoClient.h
--- a/Engine/Source/Editor/UnrealEd/Public/EditorUndoClient.h
+++ b/Engine/Source/Editor/UnrealEd/Public/EditorUndoClient.h
@@ -8,6 +8,7 @@
 
 #include "CoreMinimal.h"
 #include "Misc/ITransaction.h"
+#include "Misc/TransactionObjectEvent.h"
 
 /**
  * Interface for tools wanting to handle undo/redo operations
diff --git a/Engine/Source/Runtime/CoreUObject/Public/Misc/TransactionObjectEvent.h b/Engine/Source/Runtime/CoreUObject/Public/Misc/TransactionObjectEvent.h
--- a/Engine/Source/Runtime/CoreUObject/Public/Misc/TransactionObjectEvent.h
+++ b/Engine/Source/Runtime/CoreUObject/Public/Misc/TransactionObjectEvent.h
@@ -5,6 +5,8 @@
 #include "CoreTypes.h"
 #include "Misc/Guid.h"
 
+class ITransactionObjectAnnotation;
+
 /** Delta-change information for an object that was transacted */
 struct FTransactionObjectDeltaChange
 {
diff --git a/Engine/Source/Runtime/Engine/Private/Instances/InstancedPlacementClientInfo.cpp b/Engine/Source/Runtime/Engine/Private/Instances/InstancedPlacementClientInfo.cpp
--- a/Engine/Source/Runtime/Engine/Private/Instances/InstancedPlacementClientInfo.cpp
+++ b/Engine/Source/Runtime/Engine/Private/Instances/InstancedPlacementClientInfo.cpp
@@ -3,6 +3,7 @@
 #include "Instances/InstancedPlacementClientInfo.h"
 
 #if WITH_EDITORONLY_DATA
+#include "Misc/ITransaction.h"
 #include "Instances/InstancedPlacementPartitionActor.h"
 #include "Elements/Framework/EngineElementsLibrary.h"
 
