-----------------------------------------
Fixed having to apply snapshot twice because properties SpawnCollisionHandling and bIsSpatiallyLoaded had a diff after recreating actors

#[fyi] Jason.Walter
Fix missing header

#[fyi] Jason.Walter
Fix component still getting restored when deselected

#[fyi] Jason.Walter
Fix failing ensure when foliage has no instances

#[fyi] Jason.Walter
Fix warning caused by ChildActorTemplate

#[fyi] Jason.Walter
Fix CIS issue

#[fyi] Jason.Walter

diff --git a/Engine/Plugins/VirtualProduction/LevelSnapshots/Source/FoliageSupport/Private/FoliageSupport/FoliageRestorationInfo.cpp b/Engine/Plugins/VirtualProduction/LevelSnapshots/Source/FoliageSupport/Private/FoliageSupport/FoliageRestorationInfo.cpp
--- a/Engine/Plugins/VirtualProduction/LevelSnapshots/Source/FoliageSupport/Private/FoliageSupport/FoliageRestorationInfo.cpp
+++ b/Engine/Plugins/VirtualProduction/LevelSnapshots/Source/FoliageSupport/Private/FoliageSupport/FoliageRestorationInfo.cpp
@@ -62,6 +62,6 @@ bool UE::LevelSnapshots::Foliage::Private::FFoliageRestorationInfo::ShouldSerial
 	}
 
 	// Actor foliage is currently not supported: just restore everything
-	UE_LOG(LogLevelSnapshots, Warning, TEXT("Only static mesh foliage is supported"));
+	UE_CLOG(SavedData.GetImplType() != EFoliageImplType::StaticMesh, LogLevelSnapshots, Warning, TEXT("Only static mesh foliage is supported"));
 	return false;
 }
diff --git a/Engine/Plugins/VirtualProduction/LevelSnapshots/Source/FoliageSupport/Private/FoliageSupport/InstancedFoliageActorData.cpp b/Engine/Plugins/VirtualProduction/LevelSnapshots/Source/FoliageSupport/Private/FoliageSupport/InstancedFoliageActorData.cpp
--- a/Engine/Plugins/VirtualProduction/LevelSnapshots/Source/FoliageSupport/Private/FoliageSupport/InstancedFoliageActorData.cpp
+++ b/Engine/Plugins/VirtualProduction/LevelSnapshots/Source/FoliageSupport/Private/FoliageSupport/InstancedFoliageActorData.cpp
@@ -130,10 +130,14 @@ namespace UE::LevelSnapshots::Foliage::Private::Internal
 			if (RestorationInfo.ShouldSerializeFoliageType(FoliageInfoData))
 			{
 				FFoliageInfo* FoliageInfo = FindOrAddFoliageInfo(FoliageType, FoliageActor);
-				if (FoliageInfo
-					&& ensureAlwaysMsgf(FoliageInfoData.GetComponentName(), TEXT("ComponentName is supposed to have been saved. Investigate. Maybe you used an unsupported foliage type (only EFoliageImplType::StaticMesh is supported)?")))
+				if (FoliageInfo)
 				{
-					HandleExistingFoliageUsingRequiredComponent(FoliageActor, *FoliageInfoData.GetComponentName(), PreexistingComponentToFoliageType, FoliageType);
+					// If there were no instances, no component existed, hence no component name could be saved
+					const bool bHadAtLeastOneInstanceWhenSaved = FoliageInfoData.GetComponentName().IsSet();
+					if (bHadAtLeastOneInstanceWhenSaved)
+					{
+						HandleExistingFoliageUsingRequiredComponent(FoliageActor, *FoliageInfoData.GetComponentName(), PreexistingComponentToFoliageType, FoliageType);
+					}
 					FoliageInfoData.ApplyTo(*FoliageInfo, VersionInfo);
 				}
 			}
diff --git a/Engine/Plugins/VirtualProduction/LevelSnapshots/Source/LevelSnapshots/Private/Data/Util/SnapshotUtil.cpp b/Engine/Plugins/VirtualProduction/LevelSnapshots/Source/LevelSnapshots/Private/Data/Util/SnapshotUtil.cpp
--- a/Engine/Plugins/VirtualProduction/LevelSnapshots/Source/LevelSnapshots/Private/Data/Util/SnapshotUtil.cpp
+++ b/Engine/Plugins/VirtualProduction/LevelSnapshots/Source/LevelSnapshots/Private/Data/Util/SnapshotUtil.cpp
@@ -56,6 +56,12 @@ TOptional<FSoftObjectPath> UE::LevelSnapshots::Private::ExtractActorFromPath(con
 		FSoftObjectPath(OriginalObjectPath.GetAssetPathName(), SubPathString.Left(DotAfterActorNameIndex));
 }
 
+bool UE::LevelSnapshots::Private::IsPathToWorldObject(const FSoftObjectPath& OriginalObjectPath)
+{
+	bool bDummy;
+	return ExtractActorFromPath(OriginalObjectPath, bDummy).IsSet();
+}
+
 TOptional<int32> UE::LevelSnapshots::Private::FindDotAfterActorName(const FSoftObjectPath& OriginalObjectPath)
 {
 	const static FString PersistentLevelString("PersistentLevel.");
diff --git a/Engine/Plugins/VirtualProduction/LevelSnapshots/Source/LevelSnapshots/Private/Data/Util/SnapshotUtil.h b/Engine/Plugins/VirtualProduction/LevelSnapshots/Source/LevelSnapshots/Private/Data/Util/SnapshotUtil.h
--- a/Engine/Plugins/VirtualProduction/LevelSnapshots/Source/LevelSnapshots/Private/Data/Util/SnapshotUtil.h
+++ b/Engine/Plugins/VirtualProduction/LevelSnapshots/Source/LevelSnapshots/Private/Data/Util/SnapshotUtil.h
@@ -19,11 +19,8 @@ namespace UE::LevelSnapshots::Private
 	/** If path contains an actor, return subobject path to that actor. */
 	TOptional<FSoftObjectPath> ExtractActorFromPath(const FSoftObjectPath& OriginalObjectPath, bool& bIsPathToActorSubobject);
 
-	inline bool IsPathToWorldObject(const FSoftObjectPath& OriginalObjectPath)
-	{
-		bool bDummy;
-		return ExtractActorFromPath(OriginalObjectPath, bDummy).IsSet();
-	}
+	/** Checks whether this is a soft object path to an object in an editor world */
+	bool IsPathToWorldObject(const FSoftObjectPath& OriginalObjectPath);
 
 	/** Given a path such as "/Game/MapName.MapName:PersistentLevel.StaticMeshActor_42.SomeComponent" returns the index of "S" in "SomeComponent" */
 	TOptional<int32> FindDotAfterActorName(const FSoftObjectPath& OriginalObjectPath);
diff --git a/Engine/Plugins/VirtualProduction/LevelSnapshots/Source/LevelSnapshots/Private/Data/Util/WorldData/AllocationActorUtil.cpp b/Engine/Plugins/VirtualProduction/LevelSnapshots/Source/LevelSnapshots/Private/Data/Util/WorldData/AllocationActorUtil.cpp
--- a/Engine/Plugins/VirtualProduction/LevelSnapshots/Source/LevelSnapshots/Private/Data/Util/WorldData/AllocationActorUtil.cpp
+++ b/Engine/Plugins/VirtualProduction/LevelSnapshots/Source/LevelSnapshots/Private/Data/Util/WorldData/AllocationActorUtil.cpp
@@ -354,6 +354,7 @@ TOptional<TNonNullPtr<AActor>> UE::LevelSnapshots::Private::GetPreallocated(cons
 		SpawnParams.Name = FindObjectFast<UObject>(SpawnParams.OverrideLevel, ActorName)
 			? MakeUniqueObjectName(SpawnParams.OverrideLevel, TargetClass, ActorName)
 			: ActorName;
+		SpawnParams.bNoFail = true;
 		ActorCache.CachedSnapshotActor = SnapshotWorld->SpawnActor<AActor>(TargetClass, SpawnParams);
 
 		const bool bSpawnWasSuccess = ActorCache.CachedSnapshotActor.IsValid();
diff --git a/Engine/Plugins/VirtualProduction/LevelSnapshots/Source/LevelSnapshots/Private/Data/Util/WorldData/SnapshotObjectUtil.cpp b/Engine/Plugins/VirtualProduction/LevelSnapshots/Source/LevelSnapshots/Private/Data/Util/WorldData/SnapshotObjectUtil.cpp
--- a/Engine/Plugins/VirtualProduction/LevelSnapshots/Source/LevelSnapshots/Private/Data/Util/WorldData/SnapshotObjectUtil.cpp
+++ b/Engine/Plugins/VirtualProduction/LevelSnapshots/Source/LevelSnapshots/Private/Data/Util/WorldData/SnapshotObjectUtil.cpp
@@ -32,7 +32,8 @@ namespace UE::LevelSnapshots::Private::Internal
 		}
 		
 		// We're supposed to be dealing with an external (asset) reference, e.g. a UMaterial.
-		ensureMsgf(ExternalReference == nullptr || (!ExternalReference->IsA<AActor>() && !ExternalReference->IsA<UActorComponent>()), TEXT("Something is wrong. We just checked that the reference is not a world object but it is an actor or component."));
+		// Note that Blueprints contain AActors and they are referenced by UChildActorComponent.
+		ensureMsgf(ExternalReference == nullptr || !ExternalReference->IsInA(UWorld::StaticClass()), TEXT("Something is wrong. We just checked that the reference is not a world object but it is an actor or component."));
 		return ExternalReference;
 	}
 	
diff --git a/Engine/Plugins/VirtualProduction/LevelSnapshots/Source/LevelSnapshots/Private/Data/Util/WorldData/WorldDataUtil.cpp b/Engine/Plugins/VirtualProduction/LevelSnapshots/Source/LevelSnapshots/Private/Data/Util/WorldData/WorldDataUtil.cpp
--- a/Engine/Plugins/VirtualProduction/LevelSnapshots/Source/LevelSnapshots/Private/Data/Util/WorldData/WorldDataUtil.cpp
+++ b/Engine/Plugins/VirtualProduction/LevelSnapshots/Source/LevelSnapshots/Private/Data/Util/WorldData/WorldDataUtil.cpp
@@ -390,20 +390,25 @@ namespace UE::LevelSnapshots::Private::Internal
 					const FName ActorFName = *ActorName;
 					UClass* ActorClass = ActorCDO->GetClass();
 					FActorSpawnParameters SpawnParameters;
+					// Cannot be overriden but we want to inform the subscribers
+					SpawnParameters.bNoFail = true;
 					SpawnParameters.Name = ActorFName;
 					SpawnParameters.OverrideLevel = OverrideLevel;
-					SpawnParameters.bNoFail = true;
-					SpawnParameters.SpawnCollisionHandlingOverride = ESpawnActorCollisionHandlingMethod::AlwaysSpawn;
+					SpawnParameters.NameMode = FActorSpawnParameters::ESpawnActorNameMode::Required_ErrorAndReturnNull;
+					// Overriable properties
 					SpawnParameters.Template = ActorCDO.GetValue();
 					SpawnParameters.ObjectFlags = ActorSnapshot->SerializedActorData.GetObjectFlags();
 					
 					Module.OnPreCreateActor(OwningLevelWorld, ActorClass, SpawnParameters);
-					checkf(SpawnParameters.Name == ActorFName, TEXT("You cannot override the actor's name"));
-					checkf(SpawnParameters.OverrideLevel == OverrideLevel, TEXT("You cannot override the actor's level"))
-					
+					ensureMsgf(SpawnParameters.bNoFail, TEXT("You cannot override bNoFail"));
+					ensureMsgf(SpawnParameters.Name == ActorFName, TEXT("You cannot override the actor's name"));
+					ensureMsgf(SpawnParameters.OverrideLevel == OverrideLevel, TEXT("You cannot override the actor's level"));
+					ensureMsgf(SpawnParameters.NameMode == FActorSpawnParameters::ESpawnActorNameMode::Required_ErrorAndReturnNull, TEXT("You cannot override the actor's NameMode"));
+					SpawnParameters.bNoFail = true;
 					SpawnParameters.Name = ActorFName;
-					SpawnParameters.OverrideLevel = OwningLevelWorld->PersistentLevel;
+					SpawnParameters.OverrideLevel = OverrideLevel;
 					SpawnParameters.NameMode = FActorSpawnParameters::ESpawnActorNameMode::Required_ErrorAndReturnNull;
+					
 					if (AActor* RecreatedActor = OwningLevelWorld->SpawnActor(ActorClass, nullptr, SpawnParameters))
 					{
 						Module.OnPostRecreateActor(RecreatedActor);
diff --git a/Engine/Plugins/VirtualProduction/LevelSnapshots/Source/LevelSnapshots/Private/Restorability/EngineTypesRestorationFence.cpp b/Engine/Plugins/VirtualProduction/LevelSnapshots/Source/LevelSnapshots/Private/Restorability/EngineTypesRestorationFence.cpp
--- a/Engine/Plugins/VirtualProduction/LevelSnapshots/Source/LevelSnapshots/Private/Restorability/EngineTypesRestorationFence.cpp
+++ b/Engine/Plugins/VirtualProduction/LevelSnapshots/Source/LevelSnapshots/Private/Restorability/EngineTypesRestorationFence.cpp
@@ -5,15 +5,15 @@
 #include "ActorGroupRestoration.h"
 #include "Restorability/CollisionRestoration.h"
 #include "Restorability/GridPlacementRestoration.h"
+#include "SpecialActorPropertySupport.h"
 
-#include "EngineUtils.h"
 #include "Algo/Transform.h"
-#include "Editor/GroupActor.h"
+#include "Components/SkeletalMeshComponent.h"
+#include "EngineUtils.h"
 #include "Engine/Brush.h"
 #include "GameFramework/Actor.h"
 #include "Materials/MaterialInstance.h"
 
-
 namespace UE::LevelSnapshots::Private::EngineTypesRestorationFence
 {
 	static void AddSoftObjectPathSupport(UE::LevelSnapshots::Private::FLevelSnapshotsModule& Module)
@@ -87,6 +87,16 @@ namespace UE::LevelSnapshots::Private::EngineTypesRestorationFence
 		}
 #endif
 	}
+
+	static void DisableClothingSimulationFactory(UE::LevelSnapshots::Private::FLevelSnapshotsModule& Module)
+	{
+		// This property is directly set by RegisterComponent and similar functions... we have no chance of restoring this property properly
+		const FProperty* ClothingSimulationFactory = USkeletalMeshComponent::StaticClass()->FindPropertyByName(GET_MEMBER_NAME_CHECKED(USkeletalMeshComponent, ClothingSimulationFactory));
+		if (ensure(ClothingSimulationFactory))
+		{
+			Module.AddExplicitlyUnsupportedProperties({ ClothingSimulationFactory });
+		}
+	}
 	
 	void RegisterSpecialEngineTypeSupport(FLevelSnapshotsModule& Module)
 	{
@@ -97,10 +107,12 @@ namespace UE::LevelSnapshots::Private::EngineTypesRestorationFence
 		DisableIrrelevantWorldSettings(Module);
 		DisableIrrelevantMaterialInstanceProperties(Module);
 		DisableIrrelevantActorProperties(Module);
+		DisableClothingSimulationFactory(Module);
 
 		// Interact with special engine features
 		FCollisionRestoration::Register(Module);
 		GridPlacementRestoration::Register(Module);
 		ActorGroupRestoration::Register(Module);
+		SpecialActorPropertySupport::Register(Module);
 	}
 }
diff --git a/Engine/Plugins/VirtualProduction/LevelSnapshots/Source/LevelSnapshots/Private/Restorability/SpecialActorPropertySupport.cpp b/Engine/Plugins/VirtualProduction/LevelSnapshots/Source/LevelSnapshots/Private/Restorability/SpecialActorPropertySupport.cpp
new file mode 100644
--- /dev/null
+++ b/Engine/Plugins/VirtualProduction/LevelSnapshots/Source/LevelSnapshots/Private/Restorability/SpecialActorPropertySupport.cpp
@@ -0,0 +1,38 @@
+// Copyright Epic Games, Inc. All Rights Reserved.
+
+#include "SpecialActorPropertySupport.h"
+
+#include "LevelSnapshotsModule.h"
+#include "Params/PropertyComparisonParams.h"
+#include "Restorability/Interfaces/IPropertyComparer.h"
+
+namespace UE::LevelSnapshots::Private::SpecialActorPropertySupport
+{
+	class FSpecialActorPropertySupport : public IPropertyComparer
+	{
+		const FProperty* Property_bIsSpatiallyLoaded = AActor::StaticClass()->FindPropertyByName(AActor::GetIsSpatiallyLoadedPropertyName());
+	public:
+
+		FSpecialActorPropertySupport()
+		{
+			check(Property_bIsSpatiallyLoaded);
+		}
+		
+		virtual EPropertyComparison ShouldConsiderPropertyEqual(const FPropertyComparisonParams& Params) const override
+		{
+			if (Params.LeafProperty == Property_bIsSpatiallyLoaded)
+			{
+				return Params.WorldActor->CanChangeIsSpatiallyLoadedFlag()
+					? EPropertyComparison::CheckNormally
+					: EPropertyComparison::TreatEqual;
+			}
+			return EPropertyComparison::CheckNormally;
+		}
+	};
+	
+	void Register(FLevelSnapshotsModule& Module)
+	{
+		const TSharedRef<FSpecialActorPropertySupport> SpecialActorPropertySupport = MakeShared<FSpecialActorPropertySupport>();
+		Module.RegisterPropertyComparer(AActor::StaticClass(), SpecialActorPropertySupport);
+	}
+}
diff --git a/Engine/Plugins/VirtualProduction/LevelSnapshots/Source/LevelSnapshots/Private/Restorability/SpecialActorPropertySupport.h b/Engine/Plugins/VirtualProduction/LevelSnapshots/Source/LevelSnapshots/Private/Restorability/SpecialActorPropertySupport.h
new file mode 100644
--- /dev/null
+++ b/Engine/Plugins/VirtualProduction/LevelSnapshots/Source/LevelSnapshots/Private/Restorability/SpecialActorPropertySupport.h
@@ -0,0 +1,14 @@
+// Copyright Epic Games, Inc. All Rights Reserved.
+
+#pragma once
+
+#include "CoreMinimal.h"
+#include "LevelSnapshotsModule.h"
+
+namespace UE::LevelSnapshots::Private::SpecialActorPropertySupport
+{
+	void Register(FLevelSnapshotsModule& Module);
+}
+
+
+
diff --git a/Engine/Plugins/VirtualProduction/LevelSnapshots/Source/LevelSnapshotsEditor/Private/Views/Results/SLevelSnapshotsEditorResults.cpp b/Engine/Plugins/VirtualProduction/LevelSnapshots/Source/LevelSnapshotsEditor/Private/Views/Results/SLevelSnapshotsEditorResults.cpp
--- a/Engine/Plugins/VirtualProduction/LevelSnapshots/Source/LevelSnapshotsEditor/Private/Views/Results/SLevelSnapshotsEditorResults.cpp
+++ b/Engine/Plugins/VirtualProduction/LevelSnapshots/Source/LevelSnapshotsEditor/Private/Views/Results/SLevelSnapshotsEditorResults.cpp
@@ -500,12 +500,19 @@ void SLevelSnapshotsEditorResults::BuildSelectionSetFromSelectedPropertiesInEach
 			}
 
 			UObject* WorldObject = Group->GetWorldObject();
-
 			if (!ensureMsgf(WorldObject,
 				TEXT("%hs: WorldObject is not valid. Group name: %s"), __FUNCTION__, *Group->GetDisplayName().ToString()))
 			{
 				return;
 			}
+			
+			if (Group->GetWidgetCheckedState() == ECheckBoxState::Unchecked)
+			{
+				SelectionMap.RemoveObjectPropertiesFromMap(WorldObject);
+				SelectionMap.RemoveComponentSelection(Cast<AActor>(WorldObject));
+				SelectionMap.RemoveCustomEditorSubobjectToRecreate(WorldObject, Group->GetSnapshotObject());
+				return;
+			}
 
 			// We only want to check component or subobject groups or actor groups which have had their children generated already
 			if (Group->GetRowType() == FLevelSnapshotsEditorResultsRow::ModifiedActorGroup && !Group->GetHasGeneratedChildren())
@@ -530,7 +537,6 @@ void SLevelSnapshotsEditorResults::BuildSelectionSetFromSelectedPropertiesInEach
 
 			UE::LevelSnapshots::FAddedAndRemovedComponentInfo NewComponentSelection;
 			const UE::LevelSnapshots::FAddedAndRemovedComponentInfo* OldComponentSelection = SelectionMap.GetObjectSelection(WorldObject).GetComponentSelection();
-
 			if (OldComponentSelection)
 			{
 				NewComponentSelection.SnapshotComponentsToAdd = OldComponentSelection->SnapshotComponentsToAdd;
@@ -592,8 +598,6 @@ void SLevelSnapshotsEditorResults::BuildSelectionSetFromSelectedPropertiesInEach
 
 			if (UncheckedChildPropertyNodes.Num())
 			{
-				SelectionMap.RemoveObjectPropertiesFromMap(WorldObject);
-
 				for (const FLevelSnapshotsEditorResultsRowPtr& ChildRow : UncheckedChildPropertyNodes)
 				{
 					if (ChildRow.IsValid())
