Fix static tests failing

#[fyi] Jason.Walter
Fix snapshot incorrect rotation when parent actor removed

#[fyi] Jason.Walter

diff --git a/Engine/Plugins/VirtualProduction/LevelSnapshots/Source/LevelSnapshots/Private/Archive/ApplySnapshotToEditorArchive.cpp b/Engine/Plugins/VirtualProduction/LevelSnapshots/Source/LevelSnapshots/Private/Archive/ApplySnapshotToEditorArchive.cpp
--- a/Engine/Plugins/VirtualProduction/LevelSnapshots/Source/LevelSnapshots/Private/Archive/ApplySnapshotToEditorArchive.cpp
+++ b/Engine/Plugins/VirtualProduction/LevelSnapshots/Source/LevelSnapshots/Private/Archive/ApplySnapshotToEditorArchive.cpp
@@ -160,7 +160,7 @@ void UE::LevelSnapshots::Private::FApplySnapshotToEditorArchive::ApplyToExisting
 	TOptional<FClassDataIndex> ClassIndex)
 {
 	const FPropertySelection* Selection = InSelectionMapForResolvingSubobjects.GetObjectSelection(InOriginalObject).GetPropertySelection();
-	if (Selection && Selection->IsEmpty())
+	if (!Selection || Selection->IsEmpty())
 	{
 		return;
 	}
diff --git a/Engine/Plugins/VirtualProduction/LevelSnapshots/Source/LevelSnapshots/Private/Data/LevelSnapshot.cpp b/Engine/Plugins/VirtualProduction/LevelSnapshots/Source/LevelSnapshots/Private/Data/LevelSnapshot.cpp
--- a/Engine/Plugins/VirtualProduction/LevelSnapshots/Source/LevelSnapshots/Private/Data/LevelSnapshot.cpp
+++ b/Engine/Plugins/VirtualProduction/LevelSnapshots/Source/LevelSnapshots/Private/Data/LevelSnapshot.cpp
@@ -268,7 +268,8 @@ void ULevelSnapshot::DiffWorld(UWorld* World, FActorConsumer HandleMatchedActor,
 			const bool bWasRemovedFromWorld = ResolvedActor == nullptr || !AllActors.Contains(Cast<AActor>(ResolvedActor));
 			if (bWasRemovedFromWorld)
 			{
-				const UE::LevelSnapshots::FCanRecreateActorParams Params { World, ActorClass, OriginalActorPath, [this, &OriginalActorPath](){ return GetDeserializedActor(OriginalActorPath); } };
+				auto GetDeserializedActorFunc = [this, &OriginalActorPath](){ return GetDeserializedActor(OriginalActorPath); };
+				const UE::LevelSnapshots::FCanRecreateActorParams Params { World, ActorClass, OriginalActorPath, GetDeserializedActorFunc };
 				if (UE::LevelSnapshots::Restorability::ShouldConsiderRemovedActorForRecreation(Params))
 				{
 					HandleRemovedActor.Execute(OriginalActorPath);
diff --git a/Engine/Plugins/VirtualProduction/LevelSnapshots/Source/LevelSnapshots/Private/Data/Util/WorldData/WorldDataUtil.cpp b/Engine/Plugins/VirtualProduction/LevelSnapshots/Source/LevelSnapshots/Private/Data/Util/WorldData/WorldDataUtil.cpp
--- a/Engine/Plugins/VirtualProduction/LevelSnapshots/Source/LevelSnapshots/Private/Data/Util/WorldData/WorldDataUtil.cpp
+++ b/Engine/Plugins/VirtualProduction/LevelSnapshots/Source/LevelSnapshots/Private/Data/Util/WorldData/WorldDataUtil.cpp
@@ -455,6 +455,11 @@ namespace UE::LevelSnapshots::Private::Internal
 		// If it's not an actor nor component then it's possibly an UObjectRedirector
 		UActorComponent* AsComponent = Cast<UActorComponent>(FoundObject); // Unlikely
 		AActor* AsActor = AsComponent ? AsComponent->GetOwner() : Cast<AActor>(FoundObject);
+		if (!AsActor)
+		{
+			return false;
+		}
+		
 		const bool bShouldRename= AsActor != nullptr;
 		const bool bCanDestroy = IsValid(AsActor);
 		if (bCanDestroy)
diff --git a/Engine/Plugins/VirtualProduction/LevelSnapshots/Source/LevelSnapshots/Private/LevelSnapshotsModule.cpp b/Engine/Plugins/VirtualProduction/LevelSnapshots/Source/LevelSnapshots/Private/LevelSnapshotsModule.cpp
--- a/Engine/Plugins/VirtualProduction/LevelSnapshots/Source/LevelSnapshots/Private/LevelSnapshotsModule.cpp
+++ b/Engine/Plugins/VirtualProduction/LevelSnapshots/Source/LevelSnapshots/Private/LevelSnapshotsModule.cpp
@@ -158,7 +158,7 @@ void UE::LevelSnapshots::Private::FLevelSnapshotsModule::UnregisterCustomObjectS
 
 void UE::LevelSnapshots::Private::FLevelSnapshotsModule::RegisterGlobalActorFilter(TSharedRef<IActorSnapshotFilter> Filter)
 {
-	GlobalFilters.AddUnique(MoveTemp(Filter));
+	GlobalFilters.AddUnique(Filter);
 }
 
 void UE::LevelSnapshots::Private::FLevelSnapshotsModule::UnregisterGlobalActorFilter(const TSharedRef<IActorSnapshotFilter>& Filter)
@@ -186,6 +186,16 @@ void UE::LevelSnapshots::Private::FLevelSnapshotsModule::UnregisterRestorationLi
 	RestorationListeners.RemoveSingle(Listener);
 }
 
+void UE::LevelSnapshots::Private::FLevelSnapshotsModule::RegisterSnapshotFilterExtender(TSharedRef<ISnapshotFilterExtender> Extender)
+{
+	FilterExtenders.AddUnique(Extender);
+}
+
+void UE::LevelSnapshots::Private::FLevelSnapshotsModule::UnregisterSnapshotFilterExtender(const TSharedRef<ISnapshotFilterExtender>& Listener)
+{
+	FilterExtenders.RemoveSingle(Listener);
+}
+
 void UE::LevelSnapshots::Private::FLevelSnapshotsModule::AddExplicitilySupportedProperties(const TSet<const FProperty*>& Properties)
 {
 	for (const FProperty* Property : Properties)
@@ -529,4 +539,14 @@ void UE::LevelSnapshots::Private::FLevelSnapshotsModule::OnPostRemoveComponent(c
 	}
 }
 
+UE::LevelSnapshots::FPostApplyFiltersResult UE::LevelSnapshots::Private::FLevelSnapshotsModule::PostApplyFilters(const FPostApplyFiltersParams& Params)
+{
+	FPostApplyFiltersResult Result;
+	for (const TSharedRef<ISnapshotFilterExtender>& Extender : FilterExtenders)
+	{
+		Result.Combine(Extender->PostApplyFilters(Params));
+	}
+	return Result;
+}
+
 IMPLEMENT_MODULE(UE::LevelSnapshots::Private::FLevelSnapshotsModule, LevelSnapshots)
diff --git a/Engine/Plugins/VirtualProduction/LevelSnapshots/Source/LevelSnapshots/Private/LevelSnapshotsModule.h b/Engine/Plugins/VirtualProduction/LevelSnapshots/Source/LevelSnapshots/Private/LevelSnapshotsModule.h
--- a/Engine/Plugins/VirtualProduction/LevelSnapshots/Source/LevelSnapshots/Private/LevelSnapshotsModule.h
+++ b/Engine/Plugins/VirtualProduction/LevelSnapshots/Source/LevelSnapshots/Private/LevelSnapshotsModule.h
@@ -42,6 +42,8 @@ namespace UE::LevelSnapshots::Private
 		virtual void UnregisterSnapshotLoader(const TSharedRef<ISnapshotLoader>& Loader) override;
 		virtual void RegisterRestorationListener(TSharedRef<IRestorationListener> Listener) override;
 		virtual void UnregisterRestorationListener(const TSharedRef<IRestorationListener>& Listener) override;
+		virtual void RegisterSnapshotFilterExtender(TSharedRef<ISnapshotFilterExtender> Extender) override;
+		virtual void UnregisterSnapshotFilterExtender(const TSharedRef<ISnapshotFilterExtender>& Listener) override;
 		virtual void AddExplicitilySupportedProperties(const TSet<const FProperty*>& Properties) override;
 		virtual void RemoveAdditionallySupportedProperties(const TSet<const FProperty*>& Properties) override;
 		virtual void AddExplicitlyUnsupportedProperties(const TSet<const FProperty*>& Properties) override;
@@ -91,6 +93,8 @@ namespace UE::LevelSnapshots::Private
 
 		void OnPreRemoveComponent(UActorComponent* ComponentToRemove);
 		void OnPostRemoveComponent(const FPostRemoveComponentParams& Params);
+
+		FPostApplyFiltersResult PostApplyFilters(const FPostApplyFiltersParams& Params);
 		
 	private:
 
@@ -112,6 +116,7 @@ namespace UE::LevelSnapshots::Private
 		TArray<TSharedRef<IActorSnapshotFilter>> GlobalFilters;
 		TArray<TSharedRef<ISnapshotLoader>> SnapshotLoaders;
 		TArray<TSharedRef<IRestorationListener>> RestorationListeners;
+		TArray<TSharedRef<ISnapshotFilterExtender>> FilterExtenders;
 
 		/* Allows these properties even when the default behaviour would exclude them. */
 		TSet<const FProperty*> SupportedProperties;
diff --git a/Engine/Plugins/VirtualProduction/LevelSnapshots/Source/LevelSnapshots/Private/Restorability/AttachParentShowsTransformProperties.cpp b/Engine/Plugins/VirtualProduction/LevelSnapshots/Source/LevelSnapshots/Private/Restorability/AttachParentShowsTransformProperties.cpp
new file mode 100644
--- /dev/null
+++ b/Engine/Plugins/VirtualProduction/LevelSnapshots/Source/LevelSnapshots/Private/Restorability/AttachParentShowsTransformProperties.cpp
@@ -0,0 +1,80 @@
+// Copyright Epic Games, Inc. All Rights Reserved.
+
+#include "AttachParentShowsTransformProperties.h"
+
+#include "Interfaces/ISnapshotFilterExtender.h"
+#include "Components/SceneComponent.h"
+#include "Selection/PropertySelectionMap.h"
+
+namespace UE::LevelSnapshots::Private::AttachParentShowsTransformPropertiesFix
+{
+	/**
+	 * By default only modified properties are shown. However in the below case we show the transform properties (Location,
+	 * Rotation, and Scale) if one of the selected properties is the AttachParent.
+	 * 
+	 *	1. Create cube
+	 *	2. Rotate it to 40 degrees
+	 *  3. Take a snapshot
+	 *  4. Create new empty actor
+	 *	5. Attach cube to empty actor
+	 *	6. Rotate empty actor by 10 degrees
+	 *	7. Apply the snapshot
+	 *	
+	 *	Result: After applying, the cube has a diff because its rotation is 50 degrees.
+	 *
+	 *	The reason is because before applying the only thing that was different was the cube's attach parent. When the
+	 *	parent actor is removed, its world-space transform is retained; actors are removed first and then the properties
+	 *	of modified actors are saved. The solution is to add the cube's transform properties even though they're equal to
+	 *	the snapshot version. Users will see the property in the UI and can opt to not restore it.
+	 */
+	class FAttachParentShowsTransformProperties : public ISnapshotFilterExtender
+	{
+		const FProperty* AttachParentProperty = USceneComponent::StaticClass()->FindPropertyByName(TEXT("AttachParent"));
+		const FProperty* RelativeLocationProperty = USceneComponent::StaticClass()->FindPropertyByName(USceneComponent::GetRelativeLocationPropertyName());
+		const FProperty* RelativeRotationProperty = USceneComponent::StaticClass()->FindPropertyByName(USceneComponent::GetRelativeRotationPropertyName());
+		const FProperty* RelativeScaleProperty = USceneComponent::StaticClass()->FindPropertyByName(USceneComponent::GetRelativeScale3DPropertyName());
+		const TArray<FLevelSnapshotPropertyChain> AdditionalPropertiesToFilter = MakeAdditionalPropertiesToFilter();
+		
+		TArray<FLevelSnapshotPropertyChain> MakeAdditionalPropertiesToFilter() const
+		{
+			const FLevelSnapshotPropertyChain RelativeLocation = RelativeLocationProperty;
+			const FLevelSnapshotPropertyChain RelativeRotation = RelativeRotationProperty;
+			const FLevelSnapshotPropertyChain RelativeScale = RelativeScaleProperty;
+			TArray<FLevelSnapshotPropertyChain> Result  { RelativeLocation, RelativeRotation, RelativeScale };
+			AddStructProperties(RelativeLocation, CastField<FStructProperty>(RelativeLocationProperty), Result);
+			AddStructProperties(RelativeRotation, CastField<FStructProperty>(RelativeRotationProperty), Result);
+			AddStructProperties(RelativeScale, CastField<FStructProperty>(RelativeScaleProperty), Result);
+			return Result;
+		}
+
+		void AddStructProperties(const FLevelSnapshotPropertyChain& BaseChain, const FStructProperty* Property, TArray<FLevelSnapshotPropertyChain>& Result) const
+		{
+			for (TFieldIterator<const FProperty> RelativeLocPropIt(CastField<FStructProperty>(RelativeLocationProperty)->Struct); RelativeLocPropIt; ++RelativeLocPropIt)
+			{
+				Result.Emplace(BaseChain.MakeAppended(*RelativeLocPropIt));
+			}
+		}
+		
+	public:
+
+		FAttachParentShowsTransformProperties()
+		{
+			check(AttachParentProperty && RelativeLocationProperty && RelativeRotationProperty && RelativeScaleProperty);
+		}
+		
+		virtual FPostApplyFiltersResult PostApplyFilters(const FPostApplyFiltersParams& Params) override
+		{
+			if (Params.SelectedProperties.ShouldSerializeProperty(nullptr, AttachParentProperty))
+			{
+				return { AdditionalPropertiesToFilter, {} };
+			}
+
+			return {};
+		}
+	};
+
+	void Register(FLevelSnapshotsModule& Module)
+	{
+		Module.RegisterSnapshotFilterExtender(MakeShared<FAttachParentShowsTransformProperties>());
+	}
+}
diff --git a/Engine/Plugins/VirtualProduction/LevelSnapshots/Source/LevelSnapshots/Private/Restorability/AttachParentShowsTransformProperties.h b/Engine/Plugins/VirtualProduction/LevelSnapshots/Source/LevelSnapshots/Private/Restorability/AttachParentShowsTransformProperties.h
new file mode 100644
--- /dev/null
+++ b/Engine/Plugins/VirtualProduction/LevelSnapshots/Source/LevelSnapshots/Private/Restorability/AttachParentShowsTransformProperties.h
@@ -0,0 +1,32 @@
+// Copyright Epic Games, Inc. All Rights Reserved.
+
+#pragma once
+
+#include "CoreMinimal.h"
+#include "LevelSnapshotsModule.h"
+
+namespace UE::LevelSnapshots::Private::AttachParentShowsTransformPropertiesFix
+{
+	/**
+	 * By default only modified properties are shown. However in the below case we show the transform properties (Location,
+	 * Rotation, and Scale) if one of the selected properties is the AttachParent.
+	 * 
+	 *	1. Create cube
+	 *	2. Rotate it to 40 degrees
+	 *  3. Take a snapshot
+	 *  4. Create new empty actor
+	 *	5. Attach cube to empty actor
+	 *	6. Rotate empty actor by 10 degrees
+	 *	7. Apply the snapshot
+	 *	
+	 *	Result: After applying, the cube has a diff because its rotation is 50 degrees.
+	 *
+	 *	The reason is because before applying the only thing that was different was the cube's attach parent. When the
+	 *	parent actor is removed, its world-space transform is retained; actors are removed first and then the properties
+	 *	of modified actors are saved. The solution is to add the cube's transform properties even though they're equal to
+	 *	the snapshot version. Users will see the property in the UI and can opt to not restore it.
+	 */
+	void Register(FLevelSnapshotsModule& Module);
+}
+
+
diff --git a/Engine/Plugins/VirtualProduction/LevelSnapshots/Source/LevelSnapshots/Private/Restorability/EngineTypesRestorationFence.cpp b/Engine/Plugins/VirtualProduction/LevelSnapshots/Source/LevelSnapshots/Private/Restorability/EngineTypesRestorationFence.cpp
--- a/Engine/Plugins/VirtualProduction/LevelSnapshots/Source/LevelSnapshots/Private/Restorability/EngineTypesRestorationFence.cpp
+++ b/Engine/Plugins/VirtualProduction/LevelSnapshots/Source/LevelSnapshots/Private/Restorability/EngineTypesRestorationFence.cpp
@@ -3,6 +3,7 @@
 #include "EngineTypesRestorationFence.h"
 
 #include "ActorGroupRestoration.h"
+#include "AttachParentShowsTransformProperties.h"
 #include "Restorability/CollisionRestoration.h"
 #include "Restorability/GridPlacementRestoration.h"
 #include "SpecialActorPropertySupport.h"
@@ -114,5 +115,6 @@ namespace UE::LevelSnapshots::Private::EngineTypesRestorationFence
 		GridPlacementRestoration::Register(Module);
 		ActorGroupRestoration::Register(Module);
 		SpecialActorPropertySupport::Register(Module);
+		AttachParentShowsTransformPropertiesFix::Register(Module);
 	}
 }
diff --git a/Engine/Plugins/VirtualProduction/LevelSnapshots/Source/LevelSnapshots/Private/Selection/ApplySnapshotFilter.cpp b/Engine/Plugins/VirtualProduction/LevelSnapshots/Source/LevelSnapshots/Private/Selection/ApplySnapshotFilter.cpp
--- a/Engine/Plugins/VirtualProduction/LevelSnapshots/Source/LevelSnapshots/Private/Selection/ApplySnapshotFilter.cpp
+++ b/Engine/Plugins/VirtualProduction/LevelSnapshots/Source/LevelSnapshots/Private/Selection/ApplySnapshotFilter.cpp
@@ -181,6 +181,7 @@ void UE::LevelSnapshots::Private::FApplySnapshotFilter::FilterActorPair(FPropert
 	
 	AnalyseRootProperties(ActorContext, DeserializedSnapshotActor, WorldActor);
 	const EFilterObjectPropertiesResult FilterResult = FindAndFilterCustomSubobjectPairs(MapToAddTo, DeserializedSnapshotActor, WorldActor);
+	ExtendAnalysedProperties(ActorContext, DeserializedSnapshotActor, WorldActor);
 	
 	ActorSelection.SetHasCustomSerializedSubobjects(FilterResult == EFilterObjectPropertiesResult::HasCustomSubobjects);
 	if (!ActorSelection.IsEmpty())
@@ -221,8 +222,9 @@ UE::LevelSnapshots::Private::FApplySnapshotFilter::EPropertySearchResult UE::Lev
 	SubobjectSelection.SetHasCustomSerializedSubobjects(FilterResult == EFilterObjectPropertiesResult::HasCustomSubobjects);
 
 	// Non-component subobjects should always be added so when applying to world and resolving a subobject, we can tell whether it is referenced or just a dead object still existing in memory
-	const bool bAddedProperties = [&MapToAddTo, WorldSubobject, &SubobjectSelection]()
+	const bool bAddedProperties = [this, &MapToAddTo, &ComponentContext, SnapshotSubobject, WorldSubobject, &SubobjectSelection]()
 	{
+		ExtendAnalysedProperties(ComponentContext, SnapshotSubobject, WorldSubobject);
 		if (SubobjectSelection.IsEmpty())
 		{
 			// We don't track components because we can efficiently tell whether they're dead or not by asking the parent actor
@@ -319,7 +321,7 @@ UE::LevelSnapshots::Private::FApplySnapshotFilter::EFilterObjectPropertiesResult
 
 void UE::LevelSnapshots::Private::FApplySnapshotFilter::AnalyseRootProperties(FPropertyContainerContext& ContainerContext, UObject* SnapshotObject, UObject* WorldObject)
 {
-	FLevelSnapshotsModule& Module = FModuleManager::Get().GetModuleChecked<FLevelSnapshotsModule>("LevelSnapshots");
+	FLevelSnapshotsModule& Module = FLevelSnapshotsModule::GetInternalModuleInstance();
 	const FPropertyComparerArray PropertyComparers = Module.GetPropertyComparerForClass(ContainerContext.RootClass);
 	
 	for (TFieldIterator<FProperty> FieldIt(ContainerContext.ContainerClass); FieldIt; ++FieldIt)
@@ -345,6 +347,34 @@ void UE::LevelSnapshots::Private::FApplySnapshotFilter::AnalyseRootProperties(FP
 	}
 }
 
+void UE::LevelSnapshots::Private::FApplySnapshotFilter::ExtendAnalysedProperties(FPropertyContainerContext& ContainerContext, UObject* SnapshotObject, UObject* WorldObject)
+{
+	FLevelSnapshotsModule& Module = FLevelSnapshotsModule::GetInternalModuleInstance();
+
+	const FPostApplyFiltersResult FilterExtension = Module.PostApplyFilters({ ContainerContext.SelectionToAddTo, SnapshotObject, WorldObject });
+	for (const FLevelSnapshotPropertyChain& ToFilter : FilterExtension.AdditionalPropertiesToFilter)
+	{
+		TArray<FString> Path;
+		for (int32 i = 0; i < ToFilter.GetNumProperties(); ++i)
+		{
+			Path.Emplace(ToFilter.GetPropertyFromRoot(i)->GetAuthoredName());
+		}
+		
+		const bool bIsPropertyValid = EFilterResult::CanInclude(Filter->IsPropertyValid(
+			{ DeserializedSnapshotActor, WorldActor, ContainerContext.SnapshotContainer, ContainerContext.WorldContainer, ToFilter.GetPropertyFromRoot(ToFilter.GetNumProperties() - 1), Path }
+		));
+		if (bIsPropertyValid)
+		{
+			ContainerContext.SelectionToAddTo.AddProperty(ToFilter);
+		}
+	}
+
+	for (const FLevelSnapshotPropertyChain& ForceShow : FilterExtension.AdditionalPropertiesToForcefullyShow)
+	{
+		ContainerContext.SelectionToAddTo.AddProperty(ForceShow);
+	}
+}
+
 TOptional<UE::LevelSnapshots::Private::FApplySnapshotFilter::EPropertySearchResult> UE::LevelSnapshots::Private::FApplySnapshotFilter::HandlePossibleStructProperties(FPropertyContainerContext& ContainerContext, FProperty* PropertyToHandle)
 {
 	if (FStructProperty* StructProperty = CastField<FStructProperty>(PropertyToHandle))
diff --git a/Engine/Plugins/VirtualProduction/LevelSnapshots/Source/LevelSnapshots/Private/Selection/ApplySnapshotFilter.h b/Engine/Plugins/VirtualProduction/LevelSnapshots/Source/LevelSnapshots/Private/Selection/ApplySnapshotFilter.h
--- a/Engine/Plugins/VirtualProduction/LevelSnapshots/Source/LevelSnapshots/Private/Selection/ApplySnapshotFilter.h
+++ b/Engine/Plugins/VirtualProduction/LevelSnapshots/Source/LevelSnapshots/Private/Selection/ApplySnapshotFilter.h
@@ -80,6 +80,7 @@ namespace UE::LevelSnapshots::Private
 		EFilterObjectPropertiesResult FindAndFilterCustomSubobjectPairs(FPropertySelectionMap& MapToAddTo, UObject* SnapshotOwner, UObject* WorldOwner);
 
 		void AnalyseRootProperties(FPropertyContainerContext& ContainerContext, UObject* SnapshotObject, UObject* WorldObject);
+		void ExtendAnalysedProperties(FPropertyContainerContext& ContainerContext, UObject* SnapshotObject, UObject* WorldObject);
 		TOptional<EPropertySearchResult> HandlePossibleStructProperties(FPropertyContainerContext& ContainerContext, FProperty* PropertyToHandle);
 		EPropertySearchResult AnalyseStructProperties(FPropertyContainerContext& ContainerContext);
 
diff --git a/Engine/Plugins/VirtualProduction/LevelSnapshots/Source/LevelSnapshots/Public/ILevelSnapshotsModule.h b/Engine/Plugins/VirtualProduction/LevelSnapshots/Source/LevelSnapshots/Public/ILevelSnapshotsModule.h
--- a/Engine/Plugins/VirtualProduction/LevelSnapshots/Source/LevelSnapshots/Public/ILevelSnapshotsModule.h
+++ b/Engine/Plugins/VirtualProduction/LevelSnapshots/Source/LevelSnapshots/Public/ILevelSnapshotsModule.h
@@ -9,6 +9,7 @@
 #include "Interfaces/ISnapshotRestorabilityOverrider.h"
 #include "Interfaces/ICustomObjectSnapshotSerializer.h"
 #include "Interfaces/IActorSnapshotFilter.h"
+#include "Interfaces/ISnapshotFilterExtender.h"
 
 #include "Modules/ModuleInterface.h"
 #include "Modules/ModuleManager.h"
@@ -98,6 +99,11 @@ namespace UE::LevelSnapshots
 		/** Registers an object that will receive callbacks when a snapshot is applied. */
 		virtual void RegisterRestorationListener(TSharedRef<IRestorationListener> Listener) = 0;
 		virtual void UnregisterRestorationListener(const TSharedRef<IRestorationListener>& Listener) = 0;
+
+
+		/** Registers an object that can decide to display additional properties to the user that default snapshot behaviour does not display. */
+		virtual void RegisterSnapshotFilterExtender(TSharedRef<ISnapshotFilterExtender> Extender) = 0;
+		virtual void UnregisterSnapshotFilterExtender(const TSharedRef<ISnapshotFilterExtender>& Listener) = 0;
 		
 		
 		/**
diff --git a/Engine/Plugins/VirtualProduction/LevelSnapshots/Source/LevelSnapshots/Public/Restorability/Interfaces/ISnapshotFilterExtender.h b/Engine/Plugins/VirtualProduction/LevelSnapshots/Source/LevelSnapshots/Public/Restorability/Interfaces/ISnapshotFilterExtender.h
new file mode 100644
--- /dev/null
+++ b/Engine/Plugins/VirtualProduction/LevelSnapshots/Source/LevelSnapshots/Public/Restorability/Interfaces/ISnapshotFilterExtender.h
@@ -0,0 +1,74 @@
+// Copyright Epic Games, Inc. All Rights Reserved.
+
+#pragma once
+
+#include "CoreMinimal.h"
+#include "Selection/PropertySelection.h"
+
+class AActor;
+class UObject;
+struct FPropertySelectionMap;
+
+namespace UE::LevelSnapshots
+{
+	struct FPostApplyFiltersParams
+	{
+		const FPropertySelection& SelectedProperties;
+		UObject* SnapshotObject;
+		UObject* EditorObject;
+
+		FPostApplyFiltersParams(const FPropertySelection& SelectedProperties, UObject* SnapshotObject, UObject* EditorObject)
+			: SelectedProperties(SelectedProperties)
+			, SnapshotObject(SnapshotObject)
+			, EditorObject(EditorObject)
+		{}
+	};
+	
+	struct FPostApplyFiltersResult
+	{
+		/** Additional properties that should be displayed to the user. They will only be shown if they pass the user's filters. */
+		TArray<FLevelSnapshotPropertyChain> AdditionalPropertiesToFilter;
+		/** Additional properties that should always be displayed to the user regardless of whether they pass the property filter or not. */
+		TArray<FLevelSnapshotPropertyChain> AdditionalPropertiesToForcefullyShow;
+
+		void Combine(FPostApplyFiltersResult Result)
+		{
+			for (FLevelSnapshotPropertyChain& ToFilter : Result.AdditionalPropertiesToFilter)
+			{
+				AdditionalPropertiesToFilter.Emplace(MoveTemp(ToFilter));
+			}
+			for (FLevelSnapshotPropertyChain& ForceShow : Result.AdditionalPropertiesToForcefullyShow)
+			{
+				AdditionalPropertiesToForcefullyShow.Emplace(MoveTemp(ForceShow));
+			}
+		}
+	};
+	
+	/**
+	 * Allows you to display additional properties in the results view. This is used whenever a selection set is built.
+	 *
+	 * By default only modified properties are shown. However, sometimes you may want to include unmodified properties, too.
+	 * One example is:
+	 *	1. Create cube
+	 *	2. Rotate it to 40 degrees
+	 *  3. Take a snapshot
+	 *  4. Create new empty actor
+	 *	5. Attach cube to empty actor
+	 *	6. Rotate empty actor by 10 degrees
+	 *	7. Apply the snapshot
+	 *	Result: After applying, the cube has a diff because its rotation is 50 degrees. The reason is because before applying
+	 *	the only thing that was different was the cube's attach parent. When the parent actor is removed, its world-space transform is retained;
+	 *	actors are removed first and then the properties of modified actors are saved.
+	 *	The solution is, using this interface, to add the cube's transform properties even though they're equal to the snapshot version.
+	 *	Users will see the property in the UI and can opt to not restore it.
+	 */
+	class LEVELSNAPSHOTS_API ISnapshotFilterExtender
+	{
+	public:
+
+		/** Called after all properties have been filtered on the given object pair. */
+		virtual FPostApplyFiltersResult PostApplyFilters(const FPostApplyFiltersParams& Params) = 0;
+		
+		virtual ~ISnapshotFilterExtender() = default;
+	};
+}
\ No newline at end of file
diff --git a/Engine/Plugins/VirtualProduction/LevelSnapshots/Source/LevelSnapshotsEditor/Private/Views/Results/LevelSnapshotsEditorResultsRow.cpp b/Engine/Plugins/VirtualProduction/LevelSnapshots/Source/LevelSnapshotsEditor/Private/Views/Results/LevelSnapshotsEditorResultsRow.cpp
--- a/Engine/Plugins/VirtualProduction/LevelSnapshots/Source/LevelSnapshotsEditor/Private/Views/Results/LevelSnapshotsEditorResultsRow.cpp
+++ b/Engine/Plugins/VirtualProduction/LevelSnapshots/Source/LevelSnapshotsEditor/Private/Views/Results/LevelSnapshotsEditorResultsRow.cpp
@@ -623,16 +623,6 @@ const TArray<FLevelSnapshotsEditorResultsRowPtr>& FLevelSnapshotsEditorResultsRo
 	return ChildRows;
 }
 
-int32 FLevelSnapshotsEditorResultsRow::GetChildCount() const
-{
-	return ChildRows.Num();
-}
-
-void FLevelSnapshotsEditorResultsRow::SetChildRows(const TArray<FLevelSnapshotsEditorResultsRowPtr>& InChildRows)
-{
-	ChildRows = InChildRows;
-}
-
 void FLevelSnapshotsEditorResultsRow::AddToChildRows(const FLevelSnapshotsEditorResultsRowPtr& InRow)
 {
 	ChildRows.Add(InRow);
diff --git a/Engine/Plugins/VirtualProduction/LevelSnapshots/Source/LevelSnapshotsEditor/Private/Views/Results/LevelSnapshotsEditorResultsRow.h b/Engine/Plugins/VirtualProduction/LevelSnapshots/Source/LevelSnapshotsEditor/Private/Views/Results/LevelSnapshotsEditorResultsRow.h
--- a/Engine/Plugins/VirtualProduction/LevelSnapshots/Source/LevelSnapshotsEditor/Private/Views/Results/LevelSnapshotsEditorResultsRow.h
+++ b/Engine/Plugins/VirtualProduction/LevelSnapshots/Source/LevelSnapshotsEditor/Private/Views/Results/LevelSnapshotsEditorResultsRow.h
@@ -184,9 +184,6 @@ struct FLevelSnapshotsEditorResultsRow final : TSharedFromThis<FLevelSnapshotsEd
 	
 	/* bHasGeneratedChildren must be true to get actual children. */
 	const TArray<FLevelSnapshotsEditorResultsRowPtr>& GetChildRows() const;
-	/* bHasGeneratedChildren must be true to get an accurate value. */
-	int32 GetChildCount() const;
-	void SetChildRows(const TArray<FLevelSnapshotsEditorResultsRowPtr>& InChildRows);
 	void AddToChildRows(const FLevelSnapshotsEditorResultsRowPtr& InRow);
 	void InsertChildRowAtIndex(const FLevelSnapshotsEditorResultsRowPtr& InRow, const int32 AtIndex = 0);
 
