- Fix actors recreated from very old snapshots not getting the RF_Transactional flag set
- Fix non-transactional actors getting modified by [REDACTED_PERSON]

#[fyi] Jason.Walter
Fix unit tests

#[fyi] Jason.Walter
Added comment to better understand unit test

#[fyi] Jason.Walter
Fix transient property in struct contained by TMap causing the property to diff

#[fyi] Jason.Walter
Fix level snapshot asset getting marked as dirty when diffing snapshot

#[fyi] Jason.Walter
Fix class defaults not getting applied correctly to recreated objects

#[fyi] Jason.Walter
Fix child actors getting recreated

#[fyi] Jason.Walter

diff --git a/Engine/Plugins/VirtualProduction/LevelSnapshots/Source/LevelSnapshots/Private/Archive/ApplySnapshotToEditorArchive.cpp b/Engine/Plugins/VirtualProduction/LevelSnapshots/Source/LevelSnapshots/Private/Archive/ApplySnapshotToEditorArchive.cpp
--- a/Engine/Plugins/VirtualProduction/LevelSnapshots/Source/LevelSnapshots/Private/Archive/ApplySnapshotToEditorArchive.cpp
+++ b/Engine/Plugins/VirtualProduction/LevelSnapshots/Source/LevelSnapshots/Private/Archive/ApplySnapshotToEditorArchive.cpp
@@ -207,27 +207,19 @@ void UE::LevelSnapshots::Private::FApplySnapshotToEditorArchive::ApplyToEditorWo
 	FWorldSnapshotData& InSharedData,
 	FSnapshotDataCache& Cache,
 	UObject* InOriginalObject,
-	UObject* InDeserializedVersion,
 	const FPropertySelectionMap& InSelectionMapForResolvingSubobjects,
 	FClassDataIndex ClassIndex)
 {
 	FPropertySelection PropertiesThatWereSerialized;
 	auto TrackSerializedProperties = [&PropertiesThatWereSerialized](const FArchiveSerializedPropertyChain* Chain, const FProperty* Property) { PropertiesThatWereSerialized.AddProperty({ Chain, Property}); };
+
+	// 1. Serialize archetype first to handle the case were the archetype has changed properties since the snapshot was taken
+	const FSubobjectArchetypeFallbackInfo ClassFallbackInfo{ InOriginalObject->GetOuter(), InOriginalObject->GetFName(), InOriginalObject->GetFlags() };
+	SerializeClassDefaultsIntoSubobject(InOriginalObject, InSharedData, ClassIndex, Cache, ClassFallbackInfo);
 	
-	// Step 1: Serialise  properties that were different from CDO at time of snapshotting and that are still different from CDO
+	// Step 2: Apply the data that was different from CDO at time of snapshotting
 	FApplySnapshotToEditorArchive ApplySavedData(InObjectData, InSharedData, InOriginalObject, InSelectionMapForResolvingSubobjects, {}, TrackSerializedProperties, Cache);
 	InOriginalObject->Serialize(ApplySavedData);
-	
-	// Step 2: Serialise any remaining properties that were not covered: properties that were equal to the CDO value when the snapshot was taken but now are different from the CDO.
-	TSet<const FTextProperty*> TextProperties = Internal::CopyPastePropertiesDifferentInCDO(
-		[&PropertiesThatWereSerialized](const FArchiveSerializedPropertyChain* Chain, const FProperty* Property)
-		{
-			const bool bWasSerializedAlready = PropertiesThatWereSerialized.ShouldSerializeProperty(Chain, Property);
-			return !bWasSerializedAlready;
-		}, InOriginalObject, InDeserializedVersion);
-
-	// Step 3: Serialize FText properties that have changed in CDO since snapshot was taken. 
-	Internal::FixUpTextPropertiesDifferentInCDO(MoveTemp(TextProperties), InSharedData, ClassIndex, Cache, InOriginalObject);
 }
 
 bool UE::LevelSnapshots::Private::FApplySnapshotToEditorArchive::ShouldSkipProperty(const FProperty* InProperty) const
diff --git a/Engine/Plugins/VirtualProduction/LevelSnapshots/Source/LevelSnapshots/Private/Archive/ApplySnapshotToEditorArchive.h b/Engine/Plugins/VirtualProduction/LevelSnapshots/Source/LevelSnapshots/Private/Archive/ApplySnapshotToEditorArchive.h
--- a/Engine/Plugins/VirtualProduction/LevelSnapshots/Source/LevelSnapshots/Private/Archive/ApplySnapshotToEditorArchive.h
+++ b/Engine/Plugins/VirtualProduction/LevelSnapshots/Source/LevelSnapshots/Private/Archive/ApplySnapshotToEditorArchive.h
@@ -53,7 +53,6 @@ namespace UE::LevelSnapshots::Private
 			FWorldSnapshotData& InSharedData,
 			FSnapshotDataCache& Cache,
 			UObject* InOriginalObject,
-			UObject* InDeserializedVersion,
 			const FPropertySelectionMap& InSelectionMapForResolvingSubobjects,
 			FClassDataIndex ClassIndex
 			);
diff --git a/Engine/Plugins/VirtualProduction/LevelSnapshots/Source/LevelSnapshots/Private/Data/LevelSnapshot.cpp b/Engine/Plugins/VirtualProduction/LevelSnapshots/Source/LevelSnapshots/Private/Data/LevelSnapshot.cpp
--- a/Engine/Plugins/VirtualProduction/LevelSnapshots/Source/LevelSnapshots/Private/Data/LevelSnapshot.cpp
+++ b/Engine/Plugins/VirtualProduction/LevelSnapshots/Source/LevelSnapshots/Private/Data/LevelSnapshot.cpp
@@ -269,7 +269,7 @@ void ULevelSnapshot::DiffWorld(UWorld* World, FActorConsumer HandleMatchedActor,
 			if (bWasRemovedFromWorld)
 			{
 				auto GetDeserializedActorFunc = [this, &OriginalActorPath](){ return GetDeserializedActor(OriginalActorPath); };
-				const UE::LevelSnapshots::FCanRecreateActorParams Params { World, ActorClass, OriginalActorPath, GetDeserializedActorFunc };
+				const UE::LevelSnapshots::FCanRecreateActorParams Params{ World, ActorClass, OriginalActorPath, SavedData.SerializedActorData.GetObjectFlags(), SerializedData, GetDeserializedActorFunc };
 				if (UE::LevelSnapshots::Restorability::ShouldConsiderRemovedActorForRecreation(Params))
 				{
 					HandleRemovedActor.Execute(OriginalActorPath);
@@ -340,15 +340,18 @@ void ULevelSnapshot::EnsureWorldInitialised()
 	
 	if (RootSnapshotWorld == nullptr)
 	{
+		TransientWorldPackage = CreatePackage(*FString::Printf(TEXT("/Temp/LevelSnapshots/%s/%s"), *GetName(), *FGuid::NewGuid().ToString()));
+		TransientWorldPackage->SetFlags(RF_Transient);
+		
 		const FName RootWorldName = *ExtractPathWithoutSubobjects(MapPath).GetAssetName();
-		RootSnapshotWorld = CreateWorld(this, RootWorldName);
+		RootSnapshotWorld = CreateWorld(TransientWorldPackage, RootWorldName);
 		
 		TSet<FName> Sublevels = ExtractLevelNames(SerializedData);
 		Sublevels.Remove(RootWorldName);
 		SnapshotSublevels.Reserve(Sublevels.Num());
 		for (const FName& SublevelName : Sublevels)
 		{
-			SnapshotSublevels.Add(CreateWorld(this, SublevelName));
+			SnapshotSublevels.Add(CreateWorld(TransientWorldPackage, SublevelName));
 		}
 		
 		// Destroy our temporary world when the editor (or game) world is destroyed. Reasons:
@@ -397,7 +400,7 @@ namespace UE::LevelSnapshots::Private
 	static UWorld* CreateWorld(UObject* Outer, FName WorldName)
 	{
 		UWorld* World = NewObject<UWorld>(Outer, WorldName);
-		World->WorldType = EWorldType::EditorPreview;
+		World->WorldType = EWorldType::Inactive;
 
 		// Note: Do NOT create a FWorldContext for this world.
 		// If you do, the render thread will send render commands every tick (and crash cuz we do not init the scene below).
@@ -428,6 +431,7 @@ void ULevelSnapshot::DestroyWorld()
 		}
 		
 		CleanUpWorld();
+		TransientWorldPackage = nullptr;
 		RootSnapshotWorld = nullptr;
 		SnapshotSublevels.Reset();
 	}
@@ -464,22 +468,8 @@ void ULevelSnapshot::RecreateSnapshotWorld()
 	if (RootSnapshotWorld)
 	{
 		DestroyWorld();
-
-		if (HasAnyFlags(RF_Standalone))
-		{
-			CollectGarbage(RF_Standalone);
-		}
-		else
-		{
-			SetFlags(RF_Standalone);
-			CollectGarbage(RF_Standalone);
-			ClearFlags(RF_Standalone);
-		}
-	}
-	else
-	{
-		EnsureWorldInitialised();
 	}
+	EnsureWorldInitialised();
 }
 
 namespace UE::LevelSnapshots::Private
diff --git a/Engine/Plugins/VirtualProduction/LevelSnapshots/Source/LevelSnapshots/Private/Data/Util/EquivalenceUtil.cpp b/Engine/Plugins/VirtualProduction/LevelSnapshots/Source/LevelSnapshots/Private/Data/Util/EquivalenceUtil.cpp
--- a/Engine/Plugins/VirtualProduction/LevelSnapshots/Source/LevelSnapshots/Private/Data/Util/EquivalenceUtil.cpp
+++ b/Engine/Plugins/VirtualProduction/LevelSnapshots/Source/LevelSnapshots/Private/Data/Util/EquivalenceUtil.cpp
@@ -215,6 +215,12 @@ bool UE::LevelSnapshots::Private::HasOriginalChangedPropertiesSinceSnapshotWasTa
 	return false;
 }
 
+namespace UE::LevelSnapshots::Private
+{
+	bool AreMapPropertiesEquivalent(ULevelSnapshot* Snapshot, const FMapProperty* MapProperty, void* SnapshotValuePtr, void* WorldValuePtr, AActor* SnapshotActor, AActor* WorldActor);
+	bool AreSetPropertiesEquivalent(ULevelSnapshot* Snapshot, const FSetProperty* SetProperty, void* SnapshotValuePtr, void* WorldValuePtr, AActor* SnapshotActor, AActor* WorldActor);
+}
+
 bool UE::LevelSnapshots::Private::AreSnapshotAndOriginalPropertiesEquivalent(ULevelSnapshot* Snapshot, const FProperty* LeafProperty, void* SnapshotContainer, void* WorldContainer, AActor* SnapshotActor, AActor* WorldActor)
 {
 	// Ensure that property's flags are allowed. Skip check collection properties, e.g. FArrayProperty::Inner, etc.: inner properties do not have the same flags.
@@ -280,12 +286,12 @@ bool UE::LevelSnapshots::Private::AreSnapshotAndOriginalPropertiesEquivalent(ULe
 		
 		if (const FMapProperty* MapProperty = CastField<FMapProperty>(LeafProperty))
 		{
-			// TODO: Use custom function. Need to do something similar to UE4MapProperty_Private::IsPermutation
+			return AreMapPropertiesEquivalent(Snapshot, MapProperty, SnapshotValuePtr, WorldValuePtr, SnapshotActor, WorldActor);
 		}
 
 		if (const FSetProperty* SetProperty = CastField<FSetProperty>(LeafProperty))
 		{
-			// TODO: Use custom function. Need to do something similar to UE4SetProperty_Private::IsPermutation
+			return AreSetPropertiesEquivalent(Snapshot, SetProperty, SnapshotValuePtr, WorldValuePtr, SnapshotActor, WorldActor);
 		}
 
 		if (const FTextProperty* TextProperty = CastField<FTextProperty>(LeafProperty))
@@ -304,6 +310,54 @@ bool UE::LevelSnapshots::Private::AreSnapshotAndOriginalPropertiesEquivalent(ULe
 	return true;
 }
 
+bool UE::LevelSnapshots::Private::AreMapPropertiesEquivalent(ULevelSnapshot* Snapshot, const FMapProperty* MapProperty, void* SnapshotValuePtr, void* WorldValuePtr, AActor* SnapshotActor, AActor* WorldActor)
+{
+	// Technically we need to check permutations like done in UE4MapProperty_Private::IsPermutation... 
+	FScriptMapHelper SnapshotMap(MapProperty, SnapshotValuePtr);
+	FScriptMapHelper WorldMap(MapProperty, WorldValuePtr);
+	if (SnapshotMap.Num() != WorldMap.Num())
+	{
+		return false;
+	}
+
+	for (int32 j = 0; j < SnapshotMap.Num(); ++j)
+	{
+		void* const SnapshotPairPtr = SnapshotMap.GetPairPtr(j);
+		void* const WorldPairPtr = WorldMap.GetPairPtr(j);
+
+		const bool bAreKeysEquivalent = AreSnapshotAndOriginalPropertiesEquivalent(Snapshot, SnapshotMap.KeyProp, SnapshotPairPtr, WorldPairPtr, SnapshotActor, WorldActor);
+		const bool bAreValuesEquivalent = AreSnapshotAndOriginalPropertiesEquivalent(Snapshot, SnapshotMap.ValueProp, SnapshotPairPtr, WorldPairPtr, SnapshotActor, WorldActor);
+		if (!bAreKeysEquivalent || !bAreValuesEquivalent)
+		{
+			return false;
+		}
+	}
+	return true;
+}
+
+bool UE::LevelSnapshots::Private::AreSetPropertiesEquivalent(ULevelSnapshot* Snapshot, const FSetProperty* SetProperty, void* SnapshotValuePtr, void* WorldValuePtr, AActor* SnapshotActor, AActor* WorldActor)
+{
+	// Technically we need to check permutations like done in UE4SetProperty_Private::IsPermutation... 
+	FScriptSetHelper SnapshotMap(SetProperty, SnapshotValuePtr);
+	FScriptSetHelper WorldMap(SetProperty, WorldValuePtr);
+	if (SnapshotMap.Num() != WorldMap.Num())
+	{
+		return false;
+	}
+
+	for (int32 j = 0; j < SnapshotMap.Num(); ++j)
+	{
+		void* const SnapshotElemValuePtr = SnapshotMap.GetElementPtr(j);
+		void* const WorldElemValuePtr = WorldMap.GetElementPtr(j);
+
+		if (!AreSnapshotAndOriginalPropertiesEquivalent(Snapshot, SetProperty->ElementProp, SnapshotElemValuePtr, WorldElemValuePtr, SnapshotActor, WorldActor))
+		{
+			return false;
+		}
+	}
+	return true;
+}
+
 bool UE::LevelSnapshots::Private::AreObjectPropertiesEquivalent(ULevelSnapshot* Snapshot, const FObjectPropertyBase* ObjectProperty, void* SnapshotValuePtr, void* WorldValuePtr, AActor* SnapshotActor, AActor* WorldActor)
 {
 	// Native identity check handles:
diff --git a/Engine/Plugins/VirtualProduction/LevelSnapshots/Source/LevelSnapshots/Private/Data/Util/WorldData/RestoreActorUtil.cpp b/Engine/Plugins/VirtualProduction/LevelSnapshots/Source/LevelSnapshots/Private/Data/Util/WorldData/RestoreActorUtil.cpp
--- a/Engine/Plugins/VirtualProduction/LevelSnapshots/Source/LevelSnapshots/Private/Data/Util/WorldData/RestoreActorUtil.cpp
+++ b/Engine/Plugins/VirtualProduction/LevelSnapshots/Source/LevelSnapshots/Private/Data/Util/WorldData/RestoreActorUtil.cpp
@@ -192,7 +192,7 @@ void UE::LevelSnapshots::Private::RestoreIntoRecreatedEditorWorldActor(AActor* O
 	{
 		const FRestoreObjectScope FinishRestore = PreSubobjectRestore_EditorWorld(Deserialized, Original, WorldData, Cache, SelectedProperties, InLocalisationSnapshotPackage);
 		// Components are not recreated with any archetype -> they must go through the full serialization process 
-		FApplySnapshotToEditorArchive::ApplyToEditorWorldObjectRecreatedWithoutArchetype(SerializedCompData, WorldData, Cache, Original, Deserialized, SelectedProperties, SerializedCompData.ClassIndex);
+		FApplySnapshotToEditorArchive::ApplyToEditorWorldObjectRecreatedWithoutArchetype(SerializedCompData, WorldData, Cache, Original, SelectedProperties, SerializedCompData.ClassIndex);
 		Internal::Restore::UpdateAttachParentAttachChildren(Original);
 	};
 	Internal::Restore::DeserializeIntoEditorWorldActor(OriginalActor, ActorData, WorldData, Cache, InLocalisationSnapshotPackage, DeserializeActor, DeserializeComponent);
diff --git a/Engine/Plugins/VirtualProduction/LevelSnapshots/Source/LevelSnapshots/Private/Data/Util/WorldData/WorldDataUtil.cpp b/Engine/Plugins/VirtualProduction/LevelSnapshots/Source/LevelSnapshots/Private/Data/Util/WorldData/WorldDataUtil.cpp
--- a/Engine/Plugins/VirtualProduction/LevelSnapshots/Source/LevelSnapshots/Private/Data/Util/WorldData/WorldDataUtil.cpp
+++ b/Engine/Plugins/VirtualProduction/LevelSnapshots/Source/LevelSnapshots/Private/Data/Util/WorldData/WorldDataUtil.cpp
@@ -397,7 +397,11 @@ namespace UE::LevelSnapshots::Private::Internal
 					SpawnParameters.NameMode = FActorSpawnParameters::ESpawnActorNameMode::Required_ErrorAndReturnNull;
 					// Overriable properties
 					SpawnParameters.Template = ActorCDO.GetValue();
-					SpawnParameters.ObjectFlags = ActorSnapshot->SerializedActorData.GetObjectFlags();
+					SpawnParameters.SpawnCollisionHandlingOverride = ESpawnActorCollisionHandlingMethod::AlwaysSpawn;
+					SpawnParameters.ObjectFlags = ActorSnapshot->SerializedActorData.GetObjectFlags()
+						// Was wrongly saved for some very old snapshots - if this causes trouble at some point ... consider just deprecating the old snapshots
+						// Note that non-transactional actors cannot be removed from the world (so if you recreate and then remove using Level Snapshots it won't work)
+						| RF_Transactional;
 					
 					Module.OnPreCreateActor(OwningLevelWorld, ActorClass, SpawnParameters);
 					ensureMsgf(SpawnParameters.bNoFail, TEXT("You cannot override bNoFail"));
diff --git a/Engine/Plugins/VirtualProduction/LevelSnapshots/Source/LevelSnapshots/Private/Restorability/EngineTypesRestorationFence.cpp b/Engine/Plugins/VirtualProduction/LevelSnapshots/Source/LevelSnapshots/Private/Restorability/EngineTypesRestorationFence.cpp
--- a/Engine/Plugins/VirtualProduction/LevelSnapshots/Source/LevelSnapshots/Private/Restorability/EngineTypesRestorationFence.cpp
+++ b/Engine/Plugins/VirtualProduction/LevelSnapshots/Source/LevelSnapshots/Private/Restorability/EngineTypesRestorationFence.cpp
@@ -87,6 +87,14 @@ namespace UE::LevelSnapshots::Private::EngineTypesRestorationFence
 			Module.AddExplicitlyUnsupportedProperties({ ActorGuid });
 		}
 #endif
+
+		// Level Snapshots must spawn with AlwaysSpawn but since this is a property on the actor as well, they will show up as changed...
+		// We'll just not support this property
+		const FProperty* SpawnCollisionHandlingMethod = AActor::StaticClass()->FindPropertyByName(FName("SpawnCollisionHandlingMethod"));
+		if (ensure(SpawnCollisionHandlingMethod))
+		{
+			Module.AddExplicitlyUnsupportedProperties({ SpawnCollisionHandlingMethod });
+		}
 	}
 
 	static void DisableClothingSimulationFactory(UE::LevelSnapshots::Private::FLevelSnapshotsModule& Module)
diff --git a/Engine/Plugins/VirtualProduction/LevelSnapshots/Source/LevelSnapshots/Private/Restorability/SnapshotRestorability.cpp b/Engine/Plugins/VirtualProduction/LevelSnapshots/Source/LevelSnapshots/Private/Restorability/SnapshotRestorability.cpp
--- a/Engine/Plugins/VirtualProduction/LevelSnapshots/Source/LevelSnapshots/Private/Restorability/SnapshotRestorability.cpp
+++ b/Engine/Plugins/VirtualProduction/LevelSnapshots/Source/LevelSnapshots/Private/Restorability/SnapshotRestorability.cpp
@@ -7,6 +7,8 @@
 #include "LandscapeGizmoActor.h"
 #include "LevelSnapshotsLog.h"
 #include "LevelSnapshotsModule.h"
+#include "SnapshotCustomVersion.h"
+#include "WorldSnapshotData.h"
 
 #include "AI/NavigationSystemConfig.h"
 #include "Algo/AllOf.h"
@@ -116,7 +118,8 @@ bool UE::LevelSnapshots::Restorability::IsActorDesirableForCapture(const AActor*
 	}
 	
 	return !Actor->IsTemplate()								// Should never happen, but we never want CDOs
-			&& !Actor->HasAnyFlags(RF_Transient)				// Don't add transient actors in non-play worlds	
+			&& !Actor->HasAnyFlags(RF_Transient)				// Don't add transient actors in non-play worlds
+			&& Actor->HasAnyFlags(RF_Transactional)				// Non transactional actors are usually modified by some other system, e.g. UChildActorComponent's child actor is not transactional
 #if WITH_EDITOR
             && Actor->IsEditable()
             && Actor->IsListedInSceneOutliner() 				// Only add actors that are allowed to be selected and drawn in editor
@@ -217,7 +220,16 @@ bool UE::LevelSnapshots::Restorability::IsPropertyExplicitlySupportedForCapture(
 
 bool UE::LevelSnapshots::Restorability::ShouldConsiderRemovedActorForRecreation(const FCanRecreateActorParams& Params)
 {
-	return LevelSnapshots::Private::FLevelSnapshotsModule::GetInternalModuleInstance().CanRecreateActor(Params);
+	// Two reasons RF_Transactional is required:
+	// 1. It will not show up in Multiuser otherwise
+	// 2. If an actor is no RF_Transactional, it is usually because some system intentionally removed that flag, e.g. UChildActorComponent
+	const bool bTransacts = (Params.ObjectFlags& RF_Transactional) != 0;
+	const bool bWereFlagsSaved = Params.ObjectFlags != 0 // Objects usually have at least one flag so if it's 0 then it is very likely the flags were not saved...
+		// ... support for saving flags was added some time before ClassArchetypeRefactor. We check against ClassArchetypeRefactor for the unlikely case that in the future an object really has no flags.
+		|| Params.WorldData.SnapshotVersionInfo.GetSnapshotCustomVersion() >= LevelSnapshots::Private::FSnapshotCustomVersion::ClassArchetypeRefactor;
+	
+	return (bTransacts || !bWereFlagsSaved)
+		&& LevelSnapshots::Private::FLevelSnapshotsModule::GetInternalModuleInstance().CanRecreateActor(Params);
 }
 
 bool UE::LevelSnapshots::Restorability::ShouldConsiderNewActorForRemoval(const AActor* Actor)
diff --git a/Engine/Plugins/VirtualProduction/LevelSnapshots/Source/LevelSnapshots/Private/Tests/ApplyToWorldTests.cpp b/Engine/Plugins/VirtualProduction/LevelSnapshots/Source/LevelSnapshots/Private/Tests/ApplyToWorldTests.cpp
--- a/Engine/Plugins/VirtualProduction/LevelSnapshots/Source/LevelSnapshots/Private/Tests/ApplyToWorldTests.cpp
+++ b/Engine/Plugins/VirtualProduction/LevelSnapshots/Source/LevelSnapshots/Private/Tests/ApplyToWorldTests.cpp
@@ -276,11 +276,19 @@ namespace UE::LevelSnapshots::Private::Tests
 				TestEqual(TEXT("No other objects changed"), PropertySelection.GetKeyCount(), 1);
 				if (bRootComponentSelected)
 				{
-					TestEqual(TEXT("No other properties changed"), ComponentSelection->GetSelectedProperties().Num(), 1);
+					// 13 = AttachParent + Relative Loc & Rot & Scale + X & Y & Z of each of Loc & Rot & Scale
+					// See FAttachParentShowsTransformProperties for more info
+					TestEqual(TEXT("No other properties changed"), ComponentSelection->GetSelectedProperties().Num(), 13);
 
 					const FProperty* AttachParentProperty = USceneComponent::StaticClass()->FindPropertyByName(FName("AttachParent"));
-					ensureMsgf(AttachParentProperty, TEXT("Did property name change?"));
+					const FProperty* RelativeLocation = USceneComponent::StaticClass()->FindPropertyByName(USceneComponent::GetRelativeLocationPropertyName());
+					const FProperty* RelativeRotation = USceneComponent::StaticClass()->FindPropertyByName(USceneComponent::GetRelativeRotationPropertyName());
+					const FProperty* RelativeScale = USceneComponent::StaticClass()->FindPropertyByName(USceneComponent::GetRelativeScale3DPropertyName());
+					ensureMsgf(AttachParentProperty && RelativeLocation && RelativeRotation && RelativeScale, TEXT("Did property name change?"));
 					TestTrue(TEXT("AttachParent changed"), ComponentSelection->IsPropertySelected(nullptr, AttachParentProperty));
+					TestTrue(TEXT("RelativeLocation changed"), ComponentSelection->IsPropertySelected(nullptr, RelativeLocation));
+					TestTrue(TEXT("RelativeRotation changed"), ComponentSelection->IsPropertySelected(nullptr, RelativeRotation));
+					TestTrue(TEXT("RelativeScale changed"), ComponentSelection->IsPropertySelected(nullptr, RelativeScale));
 				}
 			})
 
diff --git a/Engine/Plugins/VirtualProduction/LevelSnapshots/Source/LevelSnapshots/Public/Data/LevelSnapshot.h b/Engine/Plugins/VirtualProduction/LevelSnapshots/Source/LevelSnapshots/Public/Data/LevelSnapshot.h
--- a/Engine/Plugins/VirtualProduction/LevelSnapshots/Source/LevelSnapshots/Public/Data/LevelSnapshot.h
+++ b/Engine/Plugins/VirtualProduction/LevelSnapshots/Source/LevelSnapshots/Public/Data/LevelSnapshot.h
@@ -130,7 +130,12 @@ private:
 	/** Callback to when an object is modified. Clears FActorSnapshotData::bHasBeenDiffed */
 	FDelegateHandle OnObjectModifiedHandle;
 
-	
+
+#if WITH_EDITORONLY_DATA
+	/** Transient package that contains the world. Exists so calls to MarkPackageDirty, called e.g. by AActor::Modify, do not mark the snapshot dirty. */
+	UPROPERTY(Transient)
+	TObjectPtr<UPackage> TransientWorldPackage;
+#endif
 	/** The root world, equivalent to the persistent world, we will be adding temporary actors to */
 	UPROPERTY(Transient)
 	UWorld* RootSnapshotWorld;
@@ -138,7 +143,6 @@ private:
 	UPROPERTY(Transient)
 	TArray<UWorld*> SnapshotSublevels;
 	
-	
 	/** The saved snapshot data */
 	UPROPERTY()
 	FWorldSnapshotData SerializedData;
diff --git a/Engine/Plugins/VirtualProduction/LevelSnapshots/Source/LevelSnapshots/Public/Data/ObjectSnapshotData.h b/Engine/Plugins/VirtualProduction/LevelSnapshots/Source/LevelSnapshots/Public/Data/ObjectSnapshotData.h
--- a/Engine/Plugins/VirtualProduction/LevelSnapshots/Source/LevelSnapshots/Public/Data/ObjectSnapshotData.h
+++ b/Engine/Plugins/VirtualProduction/LevelSnapshots/Source/LevelSnapshots/Public/Data/ObjectSnapshotData.h
@@ -16,7 +16,7 @@ struct LEVELSNAPSHOTS_API FObjectSnapshotData
 
 	/** Flags of the object that was serialized */
 	UPROPERTY()
-	uint64 ObjectFlags {};
+	uint64 ObjectFlags = RF_Transactional;
 
 	EObjectFlags GetObjectFlags() const { return static_cast<EObjectFlags>(ObjectFlags); }
 };
\ No newline at end of file
diff --git a/Engine/Plugins/VirtualProduction/LevelSnapshots/Source/LevelSnapshots/Public/Restorability/Interfaces/IActorSnapshotFilter.h b/Engine/Plugins/VirtualProduction/LevelSnapshots/Source/LevelSnapshots/Public/Restorability/Interfaces/IActorSnapshotFilter.h
--- a/Engine/Plugins/VirtualProduction/LevelSnapshots/Source/LevelSnapshots/Public/Restorability/Interfaces/IActorSnapshotFilter.h
+++ b/Engine/Plugins/VirtualProduction/LevelSnapshots/Source/LevelSnapshots/Public/Restorability/Interfaces/IActorSnapshotFilter.h
@@ -8,6 +8,7 @@
 
 class AActor;
 class UObject;
+struct FWorldSnapshotData;
 
 namespace UE::LevelSnapshots
 {
@@ -19,6 +20,8 @@ namespace UE::LevelSnapshots
 
 		UClass* Class;
 		FSoftObjectPath ActorPath;
+		EObjectFlags ObjectFlags;
+		const FWorldSnapshotData& WorldData;
 		
 		/** If you need access to the actor's data, you can call this function. */
 		FDeserializedActorGetter DeserializeFromSnapshotFunc;
