Fix recreated actors having properties left to restore

#[fyi] Jason.Walter
Fix custom subobjects not getting recreated when actor is recreated

#[fyi] Jason.Walter
Add total time to some console commands

#[fyi] Jason.Walter

diff --git a/Engine/Plugins/VirtualProduction/LevelSnapshots/Source/LevelSnapshots/Private/Data/CustomSerialization/CustomObjectSerializationWrapper.cpp b/Engine/Plugins/VirtualProduction/LevelSnapshots/Source/LevelSnapshots/Private/Data/CustomSerialization/CustomObjectSerializationWrapper.cpp
--- a/Engine/Plugins/VirtualProduction/LevelSnapshots/Source/LevelSnapshots/Private/Data/CustomSerialization/CustomObjectSerializationWrapper.cpp
+++ b/Engine/Plugins/VirtualProduction/LevelSnapshots/Source/LevelSnapshots/Private/Data/CustomSerialization/CustomObjectSerializationWrapper.cpp
@@ -122,7 +122,9 @@ namespace UE::LevelSnapshots::Private::Internal
 				}
 
 				const FCustomSubobjectRestorationInfo* RestorationInfo = SelectionMap.GetObjectSelection(EditorObject).GetCustomSubobjectSelection();
-				if (RestorationInfo && RestorationInfo->CustomSnapshotSubobjectsToRestore.Contains(SnapshotSubobject))
+				const bool bWasMissingFromExistingActor = RestorationInfo && RestorationInfo->CustomSnapshotSubobjectsToRestore.Contains(SnapshotSubobject);
+				const bool bWasActorRecreated = SelectionMap.GetDeletedActorsToRespawn().Contains(EditorSubobject->GetTypedOuter<AActor>());
+				if (bWasMissingFromExistingActor || bWasActorRecreated)
 				{
 					// Recursively check whether subobjects also have a registered ICustomObjectSnapshotSerializer
 					const FRestoreObjectScope FinishRestore = PreObjectRestore_EditorWorld(SnapshotSubobject, EditorSubobject, WorldData, Cache, SelectionMap, LocalisationSnapshotPackage,
diff --git a/Engine/Plugins/VirtualProduction/LevelSnapshots/Source/LevelSnapshots/Private/Data/Util/WorldData/RestoreActorUtil.cpp b/Engine/Plugins/VirtualProduction/LevelSnapshots/Source/LevelSnapshots/Private/Data/Util/WorldData/RestoreActorUtil.cpp
--- a/Engine/Plugins/VirtualProduction/LevelSnapshots/Source/LevelSnapshots/Private/Data/Util/WorldData/RestoreActorUtil.cpp
+++ b/Engine/Plugins/VirtualProduction/LevelSnapshots/Source/LevelSnapshots/Private/Data/Util/WorldData/RestoreActorUtil.cpp
@@ -52,7 +52,7 @@ namespace UE::LevelSnapshots::Private::Internal::Restore
 			}
 		}
 	}
-	
+
 	static void DeserializeIntoEditorWorldActor(AActor* OriginalActor, FActorSnapshotData& ActorData, FWorldSnapshotData& WorldData, FSnapshotDataCache& Cache, UPackage* InLocalisationSnapshotPackage, FSerializeActor SerializeActor, FSerializeComponent SerializeComponent)
 	{
 		const TOptional<TNonNullPtr<AActor>> Deserialized = UE::LevelSnapshots::Private::GetDeserializedActor(OriginalActor, WorldData, Cache, InLocalisationSnapshotPackage);
@@ -139,7 +139,7 @@ namespace UE::LevelSnapshots::Private::Internal::Restore
 		{
 			// Hacky way of updating AttachChildren since there is no direct way...
 			Component->UnregisterComponent();
-			Component->SetupAttachment(Component->GetAttachParent());
+			Component->SetupAttachment(Component->GetAttachParent(), Component->GetAttachSocketName());
 			// ReregisterAllComponents will be called after the actor is done restoring which will fix up everything
 		}
 	}
@@ -198,17 +198,6 @@ void UE::LevelSnapshots::Private::RestoreIntoRecreatedEditorWorldActor(AActor* O
 	Internal::Restore::DeserializeIntoEditorWorldActor(OriginalActor, ActorData, WorldData, Cache, InLocalisationSnapshotPackage, DeserializeActor, DeserializeComponent);
 
 #if WITH_EDITOR
-	// Otherwise actor will show up with internal object name, e.g. actor previously called Cube will be StaticMeshActor1
-	if (ActorData.ActorLabel.IsEmpty()) 
-	{
-		// This might be an old snapshot in which this data was not yet saved
-		OriginalActor->SetActorLabel(OriginalActor->GetName());
-	}
-	else
-	{
-		OriginalActor->SetActorLabel(ActorData.ActorLabel);
-	}
-	
 	// Recreated actors have invalid lightning cache... e.g. recreated point lights will show error image (S_LightError)
 	OriginalActor->InvalidateLightingCacheDetailed(false);
 #endif
diff --git a/Engine/Plugins/VirtualProduction/LevelSnapshots/Source/LevelSnapshots/Private/Data/Util/WorldData/SnapshotObjectUtil.cpp b/Engine/Plugins/VirtualProduction/LevelSnapshots/Source/LevelSnapshots/Private/Data/Util/WorldData/SnapshotObjectUtil.cpp
--- a/Engine/Plugins/VirtualProduction/LevelSnapshots/Source/LevelSnapshots/Private/Data/Util/WorldData/SnapshotObjectUtil.cpp
+++ b/Engine/Plugins/VirtualProduction/LevelSnapshots/Source/LevelSnapshots/Private/Data/Util/WorldData/SnapshotObjectUtil.cpp
@@ -150,6 +150,8 @@ namespace UE::LevelSnapshots::Private::Internal
 		UObject* SnapshotVersion = SubobjectCache->SnapshotObject.Get();
 		if (ResolvedObject)
 		{
+			SubobjectCache->EditorObject = ResolvedObject;
+			
 			// We can't tell whether ResolvedObject is 1: a normal referenced object or 2: a dead object that just has not been collected yet
 			// If it is referenced from somewhere (case 1), we expect it to be in the selection map.
 			// Otherwise we'll treat it as a dead object that we fully restore.
@@ -189,12 +191,12 @@ namespace UE::LevelSnapshots::Private::Internal
 			ResolvedObject = NewObject<UObject>(Outer, ExpectedClass, *ExtractLastSubobjectName(OriginalObjectPath), SubobjectData.GetObjectFlags() | RF_Transactional, Archetype.GetValue());
 			if (ensureMsgf(ResolvedObject, TEXT("Failed to allocate '%s'"), *OriginalObjectPath.ToString()))
 			{
+				SubobjectCache->EditorObject = ResolvedObject;
 				ResolvedObject->SetFlags(SubobjectData.GetObjectFlags());
 				FApplySnapshotToEditorArchive::ApplyToEditorWorldObjectRecreatedWithArchetype(SubobjectData, WorldData, Cache, ResolvedObject, SelectionMap);
 			}
 		}
 		
-		SubobjectCache->EditorObject = ResolvedObject;
 		return ResolvedObject;
 	}
 
diff --git a/Engine/Plugins/VirtualProduction/LevelSnapshots/Source/LevelSnapshots/Private/Data/Util/WorldData/WorldDataUtil.cpp b/Engine/Plugins/VirtualProduction/LevelSnapshots/Source/LevelSnapshots/Private/Data/Util/WorldData/WorldDataUtil.cpp
--- a/Engine/Plugins/VirtualProduction/LevelSnapshots/Source/LevelSnapshots/Private/Data/Util/WorldData/WorldDataUtil.cpp
+++ b/Engine/Plugins/VirtualProduction/LevelSnapshots/Source/LevelSnapshots/Private/Data/Util/WorldData/WorldDataUtil.cpp
@@ -395,6 +395,7 @@ namespace UE::LevelSnapshots::Private::Internal
 					SpawnParameters.Name = ActorFName;
 					SpawnParameters.OverrideLevel = OverrideLevel;
 					SpawnParameters.NameMode = FActorSpawnParameters::ESpawnActorNameMode::Required_ErrorAndReturnNull;
+					SpawnParameters.bDeferConstruction = true;
 					// Overriable properties
 					SpawnParameters.Template = ActorCDO.GetValue();
 					SpawnParameters.SpawnCollisionHandlingOverride = ESpawnActorCollisionHandlingMethod::AlwaysSpawn;
@@ -404,17 +405,33 @@ namespace UE::LevelSnapshots::Private::Internal
 						| RF_Transactional;
 					
 					Module.OnPreCreateActor(OwningLevelWorld, ActorClass, SpawnParameters);
-					ensureMsgf(SpawnParameters.bNoFail, TEXT("You cannot override bNoFail"));
-					ensureMsgf(SpawnParameters.Name == ActorFName, TEXT("You cannot override the actor's name"));
-					ensureMsgf(SpawnParameters.OverrideLevel == OverrideLevel, TEXT("You cannot override the actor's level"));
-					ensureMsgf(SpawnParameters.NameMode == FActorSpawnParameters::ESpawnActorNameMode::Required_ErrorAndReturnNull, TEXT("You cannot override the actor's NameMode"));
+					ensureMsgf(SpawnParameters.bNoFail
+						&& SpawnParameters.Name == ActorFName
+						&& SpawnParameters.OverrideLevel == OverrideLevel
+						&& SpawnParameters.NameMode == FActorSpawnParameters::ESpawnActorNameMode::Required_ErrorAndReturnNull
+						&& SpawnParameters.bDeferConstruction, TEXT("You cannot override these properties!"));
 					SpawnParameters.bNoFail = true;
 					SpawnParameters.Name = ActorFName;
 					SpawnParameters.OverrideLevel = OverrideLevel;
 					SpawnParameters.NameMode = FActorSpawnParameters::ESpawnActorNameMode::Required_ErrorAndReturnNull;
+					SpawnParameters.bDeferConstruction = true;
 					
 					if (AActor* RecreatedActor = OwningLevelWorld->SpawnActor(ActorClass, nullptr, SpawnParameters))
 					{
+						// The label is set while the actor is still in construction to avoid any side-effects to PostEditChangeProperty.
+#if WITH_EDITOR
+						// Otherwise actor will show up with internal object name, e.g. actor previously called Cube will be StaticMeshActor1
+						if (ActorSnapshot->ActorLabel.IsEmpty()) 
+						{
+							// This might be an old snapshot in which this data was not yet saved
+							RecreatedActor->SetActorLabel(RecreatedActor->GetName());
+						}
+						else
+						{
+							RecreatedActor->SetActorLabel(ActorSnapshot->ActorLabel);
+						}
+#endif
+						RecreatedActor->FinishSpawning(RecreatedActor->GetActorTransform());
 						Module.OnPostRecreateActor(RecreatedActor);
 						RecreatedActors.Add(OriginalRemovedActorPath, RecreatedActor);
 					}
diff --git a/Engine/Plugins/VirtualProduction/LevelSnapshots/Source/LevelSnapshots/Private/Tests/CustomObjectSerializationTests.cpp b/Engine/Plugins/VirtualProduction/LevelSnapshots/Source/LevelSnapshots/Private/Tests/CustomObjectSerializationTests.cpp
--- a/Engine/Plugins/VirtualProduction/LevelSnapshots/Source/LevelSnapshots/Private/Tests/CustomObjectSerializationTests.cpp
+++ b/Engine/Plugins/VirtualProduction/LevelSnapshots/Source/LevelSnapshots/Private/Tests/CustomObjectSerializationTests.cpp
@@ -1,11 +1,13 @@
 ﻿// Copyright Epic Games, Inc. All Rights Reserved.
 
+#include "EngineUtils.h"
 #include "LevelSnapshotsModule.h"
 #include "Interfaces/ICustomObjectSnapshotSerializer.h"
 #include "Params/ObjectSnapshotSerializationData.h"
 #include "Selection/CustomSubobjectRestorationInfo.h"
 #include "Selection/PropertySelectionMap.h"
 #include "Types/SnapshotTestActor.h"
+#include "Util/CustomSubobjectTestUtil.h"
 #include "Util/SnapshotTestRunner.h"
 
 #include "Engine/StaticMeshActor.h"
@@ -14,34 +16,9 @@
 
 namespace UE::LevelSnapshots::Private::Tests
 {
-	TSet<const FProperty*> GetSubobjectProperties()
-	{
-		TSet<const FProperty*> Properties;
-		
-		Properties.Add(ASnapshotTestActor::StaticClass()->FindPropertyByName(GET_MEMBER_NAME_CHECKED(ASnapshotTestActor, EditableInstancedSubobject_DefaultSubobject)));
-		Properties.Add(ASnapshotTestActor::StaticClass()->FindPropertyByName(GET_MEMBER_NAME_CHECKED(ASnapshotTestActor, InstancedOnlySubobject_DefaultSubobject)));
-		Properties.Add(ASnapshotTestActor::StaticClass()->FindPropertyByName(GET_MEMBER_NAME_CHECKED(ASnapshotTestActor, NakedSubobject_DefaultSubobject)));
-		Properties.Add(USubobject::StaticClass()->FindPropertyByName(GET_MEMBER_NAME_CHECKED(USubobject, NestedChild)));
-
-		return Properties;
-	}
-	
-	void DisallowSubobjectProperties()
-	{
-		FLevelSnapshotsModule& Module = FModuleManager::Get().GetModuleChecked<FLevelSnapshotsModule>("LevelSnapshots");
-		Module.AddExplicitlyUnsupportedProperties(GetSubobjectProperties());
-	}
-
-	void ReallowSubobjectProperties()
-	{
-		FLevelSnapshotsModule& Module = FModuleManager::Get().GetModuleChecked<FLevelSnapshotsModule>("LevelSnapshots");
-		Module.RemoveExplicitlyUnsupportedProperties(GetSubobjectProperties());
-	}
-
-
 	/**
-	* Tests all interface functions are called at the correct time.
-	*/
+	 * Tests all interface functions are called at the correct time.
+	 */
 	IMPLEMENT_SIMPLE_AUTOMATION_TEST(FRestoreActorCustomSubobject, "VirtualProduction.LevelSnapshots.Snapshot.CustomObjectSerialization.RestoreActorCustomSubobject", (EAutomationTestFlags::ApplicationContextMask | EAutomationTestFlags::EngineFilter));
 	bool FRestoreActorCustomSubobject::RunTest(const FString& Parameters)
 	{
@@ -552,106 +529,60 @@ namespace UE::LevelSnapshots::Private::Tests
 	IMPLEMENT_SIMPLE_AUTOMATION_TEST(FFilterForPropertiesOnSubobjects, "VirtualProduction.LevelSnapshots.Snapshot.CustomObjectSerialization.FilterForPropertiesOnSubobjects", (EAutomationTestFlags::ApplicationContextMask | EAutomationTestFlags::EngineFilter));
 	bool FFilterForPropertiesOnSubobjects::RunTest(const FString& Parameters)
 	{
-		class FStub : public ICustomObjectSnapshotSerializer
-		{
-		public:
-		
-			virtual void OnTakeSnapshot(UObject* EditorObject, ICustomSnapshotSerializationData& DataStorage) override
-			{
-				DataStorage.AddSubobjectSnapshot(TestActor->InstancedOnlySubobject_DefaultSubobject);
-			}
-
-			virtual UObject* FindOrRecreateSubobjectInSnapshotWorld(UObject* SnapshotObject, const ISnapshotSubobjectMetaData& ObjectData, const ICustomSnapshotSerializationData& DataStorage) override
-			{
-				if (ASnapshotTestActor* SnapshotActor = Cast<ASnapshotTestActor>(SnapshotObject))
-				{
-					return SnapshotActor->InstancedOnlySubobject_DefaultSubobject;
-				}
-
-				checkNoEntry();
-				return nullptr;
-			}
-		
-			virtual UObject* FindOrRecreateSubobjectInEditorWorld(UObject* EditorObject, const ISnapshotSubobjectMetaData& ObjectData, const ICustomSnapshotSerializationData& DataStorage) override
-			{
-				return TestActor->InstancedOnlySubobject_DefaultSubobject;	
-			}
-		
-			virtual UObject* FindSubobjectInEditorWorld(UObject* EditorObject, const ISnapshotSubobjectMetaData& ObjectData, const ICustomSnapshotSerializationData& DataStorage) override
-			{
-				return TestActor->InstancedOnlySubobject_DefaultSubobject;	
-			}
-
-			ASnapshotTestActor* TestActor;
-		};
-	
 		// Handle registering and unregistering of custom serializer
-		TSharedRef<FStub> Stub = MakeShared<FStub>();
-		FLevelSnapshotsModule& Module = FModuleManager::Get().GetModuleChecked<FLevelSnapshotsModule>("LevelSnapshots");
-		Module.RegisterCustomObjectSerializer(ASnapshotTestActor::StaticClass(), Stub);
-		DisallowSubobjectProperties();
-		ON_SCOPE_EXIT
-		{
-			Module.UnregisterCustomObjectSerializer(ASnapshotTestActor::StaticClass());
-			ReallowSubobjectProperties();
-		};
-
-
-
+		const TCustomObjectSerializerContext<FInstancedOnlySubobjectCustomObjectSerializer> Stub =
+			FInstancedOnlySubobjectCustomObjectSerializer::Make();
 	
 		// Modify subobject. Actor unchanged. Actor is in selection map.
 		FSnapshotTestRunner()
 			.ModifyWorld([&](UWorld* World)
 			{
-				Stub->TestActor = ASnapshotTestActor::Spawn(World, "TestActor");
+				Stub.GetCustomSerializer()->TestActor = ASnapshotTestActor::Spawn(World, "TestActor");
 			
-				Stub->TestActor->InstancedOnlySubobject_DefaultSubobject->FloatProperty = 42.f;
+				Stub.GetCustomSerializer()->TestActor->InstancedOnlySubobject_DefaultSubobject->FloatProperty = 42.f;
 			})
 			.TakeSnapshot()
 	
 			.ModifyWorld([&](UWorld* World)
 			{
-				Stub->TestActor->InstancedOnlySubobject_DefaultSubobject->FloatProperty = 420.f;
+				Stub.GetCustomSerializer()->TestActor->InstancedOnlySubobject_DefaultSubobject->FloatProperty = 420.f;
 			})
 
-			.FilterProperties(Stub->TestActor, [&](const FPropertySelectionMap& PropertySelectionMap)
+			.FilterProperties(Stub.GetCustomSerializer()->TestActor, [&](const FPropertySelectionMap& PropertySelectionMap)
 			{
 				// Custom subobject properties
 				UClass* SubobjectClass = USubobject::StaticClass();
 				const FProperty* ChangedSubobjectProperty = SubobjectClass->FindPropertyByName(GET_MEMBER_NAME_CHECKED(USubobject, FloatProperty));
 
-				const FPropertySelection* SelectedSubobjectProperties = PropertySelectionMap.GetObjectSelection(Stub->TestActor->InstancedOnlySubobject_DefaultSubobject).GetPropertySelection();
+				const FPropertySelection* SelectedSubobjectProperties = PropertySelectionMap.GetObjectSelection(Stub.GetCustomSerializer()->TestActor->InstancedOnlySubobject_DefaultSubobject).GetPropertySelection();
 				const bool bSubobjectHasExpectedNumChangedProperties = SelectedSubobjectProperties && SelectedSubobjectProperties->GetSelectedLeafProperties().Num() == 1;
 				TestTrue(TEXT("Subobject has changed properties"), bSubobjectHasExpectedNumChangedProperties);
 				TestTrue(TEXT("Changed property on subobject contained"), bSubobjectHasExpectedNumChangedProperties && SelectedSubobjectProperties->IsPropertySelected(nullptr, ChangedSubobjectProperty));
 
 
 				// Actor properties
-				const FPropertySelection* SelectedActorProperties = PropertySelectionMap.GetObjectSelection(Stub->TestActor).GetPropertySelection();
+				const FPropertySelection* SelectedActorProperties = PropertySelectionMap.GetObjectSelection(Stub.GetCustomSerializer()->TestActor).GetPropertySelection();
 				TestTrue(TEXT("Unchanged actor in selection map when subobject was changed"), SelectedActorProperties != nullptr);
 			});
-
-
-
-
+		
 
 	
 		// Modify nothing. No properties in selection map.
 		FSnapshotTestRunner()
 			.ModifyWorld([&](UWorld* World)
 			{
-				Stub->TestActor = ASnapshotTestActor::Spawn(World, "TestActor");
+				Stub.GetCustomSerializer()->TestActor = ASnapshotTestActor::Spawn(World, "TestActor");
 			})
 			.TakeSnapshot()
-			.FilterProperties(Stub->TestActor, [&](const FPropertySelectionMap& PropertySelectionMap)
+			.FilterProperties(Stub.GetCustomSerializer()->TestActor, [&](const FPropertySelectionMap& PropertySelectionMap)
 			{
 				// Custom subobject properties
-				const FPropertySelection* SelectedSubobjectProperties = PropertySelectionMap.GetObjectSelection(Stub->TestActor->InstancedOnlySubobject_DefaultSubobject).GetPropertySelection();
+				const FPropertySelection* SelectedSubobjectProperties = PropertySelectionMap.GetObjectSelection(Stub.GetCustomSerializer()->TestActor->InstancedOnlySubobject_DefaultSubobject).GetPropertySelection();
 				TestTrue(TEXT("Unchanged subobject not in selection map"), SelectedSubobjectProperties == nullptr || SelectedSubobjectProperties->IsEmpty());
 
 
 				// Actor properties
-				const FPropertySelection* SelectedActorProperties = PropertySelectionMap.GetObjectSelection(Stub->TestActor).GetPropertySelection();
+				const FPropertySelection* SelectedActorProperties = PropertySelectionMap.GetObjectSelection(Stub.GetCustomSerializer()->TestActor).GetPropertySelection();
 				TestTrue(TEXT("Unchanged actor not in selection map when subobject was not changed"), SelectedActorProperties == nullptr || SelectedActorProperties->IsEmpty());
 			});
 
@@ -665,28 +596,28 @@ namespace UE::LevelSnapshots::Private::Tests
 		FSnapshotTestRunner()
 			.ModifyWorld([&](UWorld* World)
 			{
-				Stub->TestActor = ASnapshotTestActor::Spawn(World, "TestActor");
+				Stub.GetCustomSerializer()->TestActor = ASnapshotTestActor::Spawn(World, "TestActor");
 			
-				Stub->TestActor->InstancedOnlySubobject_DefaultSubobject->FloatProperty = 42.f;
-				Stub->TestActor->IntProperty = 42;
-				Stub->TestActor->TestComponent->IntProperty = 42;
+				Stub.GetCustomSerializer()->TestActor->InstancedOnlySubobject_DefaultSubobject->FloatProperty = 42.f;
+				Stub.GetCustomSerializer()->TestActor->IntProperty = 42;
+				Stub.GetCustomSerializer()->TestActor->TestComponent->IntProperty = 42;
 			})
 			.TakeSnapshot()
 	
 			.ModifyWorld([&](UWorld* World)
 			{
-				Stub->TestActor->InstancedOnlySubobject_DefaultSubobject->FloatProperty = 420.f;
-				Stub->TestActor->IntProperty = 420;
-				Stub->TestActor->TestComponent->IntProperty = 420;
+				Stub.GetCustomSerializer()->TestActor->InstancedOnlySubobject_DefaultSubobject->FloatProperty = 420.f;
+				Stub.GetCustomSerializer()->TestActor->IntProperty = 420;
+				Stub.GetCustomSerializer()->TestActor->TestComponent->IntProperty = 420;
 			})
 
-			.FilterProperties(Stub->TestActor, [&](const FPropertySelectionMap& PropertySelectionMap)
+			.FilterProperties(Stub.GetCustomSerializer()->TestActor, [&](const FPropertySelectionMap& PropertySelectionMap)
 			{
 				// Custom subobject properties
 				UClass* SubobjectClass = USubobject::StaticClass();
 				const FProperty* ChangedSubobjectProperty = SubobjectClass->FindPropertyByName(GET_MEMBER_NAME_CHECKED(USubobject, FloatProperty));
 
-				const FPropertySelection* SelectedSubobjectProperties = PropertySelectionMap.GetObjectSelection(Stub->TestActor->InstancedOnlySubobject_DefaultSubobject).GetPropertySelection();
+				const FPropertySelection* SelectedSubobjectProperties = PropertySelectionMap.GetObjectSelection(Stub.GetCustomSerializer()->TestActor->InstancedOnlySubobject_DefaultSubobject).GetPropertySelection();
 				const bool bSubobjectHasExpectedNumChangedProperties = SelectedSubobjectProperties && SelectedSubobjectProperties->GetSelectedLeafProperties().Num() == 1;
 				TestTrue(TEXT("Subobject has changed properties"), bSubobjectHasExpectedNumChangedProperties);
 				TestTrue(TEXT("Changed property on subobject contained"), bSubobjectHasExpectedNumChangedProperties && SelectedSubobjectProperties->IsPropertySelected(nullptr, ChangedSubobjectProperty));
@@ -696,7 +627,7 @@ namespace UE::LevelSnapshots::Private::Tests
 				UClass* TestActorClass = ASnapshotTestActor::StaticClass();
 				const FProperty* ChangedActorProperty = TestActorClass->FindPropertyByName(GET_MEMBER_NAME_CHECKED(ASnapshotTestActor, IntProperty));
 			
-				const FPropertySelection* SelectedActorProperties = PropertySelectionMap.GetObjectSelection(Stub->TestActor).GetPropertySelection();
+				const FPropertySelection* SelectedActorProperties = PropertySelectionMap.GetObjectSelection(Stub.GetCustomSerializer()->TestActor).GetPropertySelection();
 				const bool bActorHasExpectedNumChangedProperties = SelectedActorProperties && SelectedActorProperties->GetSelectedLeafProperties().Num();
 				TestTrue(TEXT("Actor has changed properties"), bActorHasExpectedNumChangedProperties);
 				TestTrue(TEXT("Changed property on actor contained"), bActorHasExpectedNumChangedProperties && SelectedActorProperties->IsPropertySelected(nullptr, ChangedActorProperty));
@@ -706,7 +637,7 @@ namespace UE::LevelSnapshots::Private::Tests
 				UClass* TestComponentClass = USnapshotTestComponent::StaticClass();
 				const FProperty* ChangedComponentProperty = TestComponentClass->FindPropertyByName(GET_MEMBER_NAME_CHECKED(USnapshotTestComponent, IntProperty));
 			
-				const FPropertySelection* SelectedComponentProperties = PropertySelectionMap.GetObjectSelection(Stub->TestActor->TestComponent).GetPropertySelection();
+				const FPropertySelection* SelectedComponentProperties = PropertySelectionMap.GetObjectSelection(Stub.GetCustomSerializer()->TestActor->TestComponent).GetPropertySelection();
 				const bool bComponentHasExpectedNumChangedProperties = SelectedComponentProperties && SelectedComponentProperties->GetSelectedLeafProperties().Num();
 				TestTrue(TEXT("Component has changed properties"), bComponentHasExpectedNumChangedProperties);
 				TestTrue(TEXT("Changed property on component contained"), bComponentHasExpectedNumChangedProperties && SelectedComponentProperties->IsPropertySelected(nullptr, ChangedComponentProperty));
@@ -716,8 +647,8 @@ namespace UE::LevelSnapshots::Private::Tests
 	}
 
 	/**
-	* Checks that restoring subobjects which are missing from the editor world are in fact restored.
-	*/
+	 * Checks that restoring subobjects which are missing from the editor world are in fact restored.
+	 */
 	IMPLEMENT_SIMPLE_AUTOMATION_TEST(FRestoreSubobjectsMissingFromEditorWorld, "VirtualProduction.LevelSnapshots.Snapshot.CustomObjectSerialization.RestoreSubobjectsMissingFromEditorWorld", (EAutomationTestFlags::ApplicationContextMask | EAutomationTestFlags::EngineFilter));
 	bool FRestoreSubobjectsMissingFromEditorWorld::RunTest(const FString& Parameters)
 	{
@@ -824,4 +755,47 @@ namespace UE::LevelSnapshots::Private::Tests
 	
 		return true;
 	}
-}
\ No newline at end of file
+
+	/**
+	 * When an actor is recreated, all of its custom subobjects are recreated as well.
+	 */
+	IMPLEMENT_SIMPLE_AUTOMATION_TEST(FRestoresCustomSubobjectWhenActorRecreated, "VirtualProduction.LevelSnapshots.Snapshot.CustomObjectSerialization.RestoresCustomSubobjectWhenActorRecreated", (EAutomationTestFlags::ApplicationContextMask | EAutomationTestFlags::EngineFilter));
+	bool FRestoresCustomSubobjectWhenActorRecreated::RunTest(const FString& Parameters)
+	{
+		// Handle registering and unregistering of custom serializer
+		// Handle registering and unregistering of custom serializer
+		const TCustomObjectSerializerContext<FInstancedOnlySubobjectCustomObjectSerializer> Stub =
+			FInstancedOnlySubobjectCustomObjectSerializer::Make();
+
+		FSnapshotTestRunner()
+			.ModifyWorld([&](UWorld* World)
+			{
+				Stub.GetCustomSerializer()->TestActor = ASnapshotTestActor::Spawn(World, "RecreatedActor");
+				Stub.GetCustomSerializer()->TestActor->InstancedOnlySubobject_DefaultSubobject->FloatProperty = 12345.f;
+			})
+			.TakeSnapshot()
+	
+			.ModifyWorld([&](UWorld* World)
+			{
+				Stub.GetCustomSerializer()->TestActor->InstancedOnlySubobject_DefaultSubobject->FloatProperty = 5.f;
+				World->DestroyActor(Stub.GetCustomSerializer()->TestActor);
+				Stub.GetCustomSerializer()->TestActor = nullptr;
+			})
+			.ApplySnapshot()
+		
+			.ModifyWorld([&](UWorld* World)
+			{
+				const TActorIterator<ASnapshotTestActor> ActorIterator(World);
+				if (!ActorIterator)
+				{
+					AddError(TEXT("Actor not restored"));
+					return;
+				}
+
+				const ASnapshotTestActor* RestoredActor = *ActorIterator;
+				TestEqual(TEXT("Default Subobject was restored"), RestoredActor->InstancedOnlySubobject_DefaultSubobject->FloatProperty, 12345.f);
+			});
+		
+		return true;
+	} 
+}
diff --git a/Engine/Plugins/VirtualProduction/LevelSnapshots/Source/LevelSnapshots/Private/Tests/Util/CustomSubobjectTestUtil.cpp b/Engine/Plugins/VirtualProduction/LevelSnapshots/Source/LevelSnapshots/Private/Tests/Util/CustomSubobjectTestUtil.cpp
new file mode 100644
--- /dev/null
+++ b/Engine/Plugins/VirtualProduction/LevelSnapshots/Source/LevelSnapshots/Private/Tests/Util/CustomSubobjectTestUtil.cpp
@@ -0,0 +1,82 @@
+// Copyright Epic Games, Inc. All Rights Reserved.
+
+#include "CustomSubobjectTestUtil.h"
+
+#include "LevelSnapshotsModule.h"
+#include "Params/ObjectSnapshotSerializationData.h"
+#include "Tests/Types/SnapshotTestActor.h"
+
+namespace UE::LevelSnapshots::Private::Tests
+{
+	TCustomObjectSerializerContext<FInstancedOnlySubobjectCustomObjectSerializer> FInstancedOnlySubobjectCustomObjectSerializer::Make()
+	{
+		return TCustomObjectSerializerContext<FInstancedOnlySubobjectCustomObjectSerializer>();
+	}
+
+	void FInstancedOnlySubobjectCustomObjectSerializer::OnTakeSnapshot(UObject* EditorObject, ICustomSnapshotSerializationData& DataStorage)
+	{
+		DataStorage.AddSubobjectSnapshot(TestActor->InstancedOnlySubobject_DefaultSubobject);
+	}
+
+	UObject* FInstancedOnlySubobjectCustomObjectSerializer::FindOrRecreateSubobjectInSnapshotWorld(UObject* SnapshotObject, const ISnapshotSubobjectMetaData& ObjectData, const ICustomSnapshotSerializationData& DataStorage)
+	{
+		if (ASnapshotTestActor* SnapshotActor = Cast<ASnapshotTestActor>(SnapshotObject))
+		{
+			check(SnapshotActor->InstancedOnlySubobject_DefaultSubobject);
+			return SnapshotActor->InstancedOnlySubobject_DefaultSubobject;
+		}
+
+		checkNoEntry();
+		return nullptr;
+	}
+
+	UObject* FInstancedOnlySubobjectCustomObjectSerializer::FindOrRecreateSubobjectInEditorWorld(UObject* EditorObject, const ISnapshotSubobjectMetaData& ObjectData, const ICustomSnapshotSerializationData& DataStorage)
+	{
+		if (ASnapshotTestActor* SnapshotActor =Cast<ASnapshotTestActor>(EditorObject))
+		{
+			check(SnapshotActor->InstancedOnlySubobject_DefaultSubobject);
+			return SnapshotActor->InstancedOnlySubobject_DefaultSubobject;
+		}
+
+		checkNoEntry();
+		return nullptr;
+	}
+
+	UObject* FInstancedOnlySubobjectCustomObjectSerializer::FindSubobjectInEditorWorld(UObject* EditorObject, const ISnapshotSubobjectMetaData& ObjectData, const ICustomSnapshotSerializationData& DataStorage)
+	{
+		if (ASnapshotTestActor* SnapshotActor = Cast<ASnapshotTestActor>(EditorObject))
+		{
+			check(SnapshotActor->InstancedOnlySubobject_DefaultSubobject);
+			return SnapshotActor->InstancedOnlySubobject_DefaultSubobject;
+		}
+
+		checkNoEntry();
+		return nullptr;
+	}
+
+	TSet<const FProperty*> GetSubobjectProperties()
+	{
+		TSet<const FProperty*> Properties;
+			
+		Properties.Add(ASnapshotTestActor::StaticClass()->FindPropertyByName(GET_MEMBER_NAME_CHECKED(ASnapshotTestActor, EditableInstancedSubobject_DefaultSubobject)));
+		Properties.Add(ASnapshotTestActor::StaticClass()->FindPropertyByName(GET_MEMBER_NAME_CHECKED(ASnapshotTestActor, InstancedOnlySubobject_DefaultSubobject)));
+		Properties.Add(ASnapshotTestActor::StaticClass()->FindPropertyByName(GET_MEMBER_NAME_CHECKED(ASnapshotTestActor, NakedSubobject_DefaultSubobject)));
+		Properties.Add(USubobject::StaticClass()->FindPropertyByName(GET_MEMBER_NAME_CHECKED(USubobject, NestedChild)));
+
+		return Properties;
+	}
+		
+	void DisallowSubobjectProperties()
+	{
+		FLevelSnapshotsModule& Module = FModuleManager::Get().GetModuleChecked<Private::FLevelSnapshotsModule>("LevelSnapshots");
+		Module.AddExplicitlyUnsupportedProperties(GetSubobjectProperties());
+	}
+
+	void ReallowSubobjectProperties()
+	{
+		FLevelSnapshotsModule& Module = FModuleManager::Get().GetModuleChecked<FLevelSnapshotsModule>("LevelSnapshots");
+		Module.RemoveExplicitlyUnsupportedProperties(GetSubobjectProperties());
+	}
+}
+
+
diff --git a/Engine/Plugins/VirtualProduction/LevelSnapshots/Source/LevelSnapshots/Private/Tests/Util/CustomSubobjectTestUtil.h b/Engine/Plugins/VirtualProduction/LevelSnapshots/Source/LevelSnapshots/Private/Tests/Util/CustomSubobjectTestUtil.h
new file mode 100644
--- /dev/null
+++ b/Engine/Plugins/VirtualProduction/LevelSnapshots/Source/LevelSnapshots/Private/Tests/Util/CustomSubobjectTestUtil.h
@@ -0,0 +1,56 @@
+// Copyright Epic Games, Inc. All Rights Reserved.
+
+#pragma once
+
+#include "CoreMinimal.h"
+#include "LevelSnapshotsModule.h"
+#include "Interfaces/ICustomObjectSnapshotSerializer.h"
+#include "Tests/Types/SnapshotTestActor.h"
+
+class ASnapshotTestActor;
+
+namespace UE::LevelSnapshots::Private::Tests
+{
+	TSet<const FProperty*> GetSubobjectProperties();
+	void DisallowSubobjectProperties();
+	void ReallowSubobjectProperties();
+	
+	template<typename TSerializerType>
+	class TCustomObjectSerializerContext
+	{
+		TSharedRef<TSerializerType> CustomSerializer;
+	public:
+		
+		template<typename ...T>
+		TCustomObjectSerializerContext(T&&... Arg)
+			: CustomSerializer(MakeShared<TSerializerType>(Forward<T>(Arg)...))
+		{
+			FLevelSnapshotsModule& Module = FModuleManager::Get().GetModuleChecked<FLevelSnapshotsModule>("LevelSnapshots");
+			Module.RegisterCustomObjectSerializer(ASnapshotTestActor::StaticClass(), CustomSerializer);
+			DisallowSubobjectProperties();
+		}
+		
+		~TCustomObjectSerializerContext()
+		{
+			FLevelSnapshotsModule& Module = FModuleManager::Get().GetModuleChecked<FLevelSnapshotsModule>("LevelSnapshots");
+			Module.UnregisterCustomObjectSerializer(ASnapshotTestActor::StaticClass());
+			ReallowSubobjectProperties();
+		}
+		
+		TSharedRef<TSerializerType> GetCustomSerializer() const { return CustomSerializer; }
+	};
+	
+	class FInstancedOnlySubobjectCustomObjectSerializer : public ICustomObjectSnapshotSerializer
+	{
+	public:
+
+		static TCustomObjectSerializerContext<FInstancedOnlySubobjectCustomObjectSerializer> Make();
+		
+		virtual void OnTakeSnapshot(UObject* EditorObject, ICustomSnapshotSerializationData& DataStorage) override;
+		virtual UObject* FindOrRecreateSubobjectInSnapshotWorld(UObject* SnapshotObject, const ISnapshotSubobjectMetaData& ObjectData, const ICustomSnapshotSerializationData& DataStorage) override;
+		virtual UObject* FindOrRecreateSubobjectInEditorWorld(UObject* EditorObject, const ISnapshotSubobjectMetaData& ObjectData, const ICustomSnapshotSerializationData& DataStorage) override;
+		virtual UObject* FindSubobjectInEditorWorld(UObject* EditorObject, const ISnapshotSubobjectMetaData& ObjectData, const ICustomSnapshotSerializationData& DataStorage) override;
+
+		ASnapshotTestActor* TestActor;
+	};
+}
diff --git a/Engine/Plugins/VirtualProduction/LevelSnapshots/Source/LevelSnapshots/Private/Util/SortedScopedLog.cpp b/Engine/Plugins/VirtualProduction/LevelSnapshots/Source/LevelSnapshots/Private/Util/SortedScopedLog.cpp
--- a/Engine/Plugins/VirtualProduction/LevelSnapshots/Source/LevelSnapshots/Private/Util/SortedScopedLog.cpp
+++ b/Engine/Plugins/VirtualProduction/LevelSnapshots/Source/LevelSnapshots/Private/Util/SortedScopedLog.cpp
@@ -62,10 +62,11 @@ void FConditionalSortedScopedLog::FFormatCSV::operator()(FOutputDevice& OutputDe
 	OutputDevice.Logf(ELogVerbosity::Log, TEXT("%s"), *Result);
 }
 
-FConditionalSortedScopedLog::FConditionalSortedScopedLog(bool bCondition, FString Title, FLogFormatter Formatter, FOutputDevice* OutputDevice, bool bLargeTimesFirst, EScopeLogTimeUnits Units)
+FConditionalSortedScopedLog::FConditionalSortedScopedLog(bool bCondition, FString Title, FLogFormatter Formatter, FOutputDevice* OutputDevice, bool bLargeTimesFirst, bool bLogTotalTime, EScopeLogTimeUnits Units)
 	: bCondition(bCondition)
 	, LogTitle(Title)
 	, bLargeTimesFirst(bLargeTimesFirst)
+	, StartTime(bLogTotalTime ? FPlatformTime::Seconds() : 0.f)
 	, Units(Units)
 	, Formatter(Formatter)
 	, Output(OutputDevice)
@@ -84,6 +85,12 @@ FConditionalSortedScopedLog::~FConditionalSortedScopedLog()
 			LoggedItems.Sort(TLess<FLoggedItem>());
 		}
 
+		const bool bLogTotalTime = StartTime != 0.0;
+		if (bLogTotalTime)
+		{
+			AddCustomMeasuredLogItem(TEXT("Total Time"), FPlatformTime::Seconds() - StartTime);
+		}
+		
 		Formatter(Output ? *Output : *GLog, GetDisplayUnitsString(), LogTitle, LoggedItems);
 	}
 }
diff --git a/Engine/Plugins/VirtualProduction/LevelSnapshots/Source/LevelSnapshots/Private/Util/SortedScopedLog.h b/Engine/Plugins/VirtualProduction/LevelSnapshots/Source/LevelSnapshots/Private/Util/SortedScopedLog.h
--- a/Engine/Plugins/VirtualProduction/LevelSnapshots/Source/LevelSnapshots/Private/Util/SortedScopedLog.h
+++ b/Engine/Plugins/VirtualProduction/LevelSnapshots/Source/LevelSnapshots/Private/Util/SortedScopedLog.h
@@ -58,11 +58,19 @@ public:
 	 * @param bCondition Determines whether anything will be printed or not. Typical use case: this condition depends on a console variable.
 	 * @param Title This will be printed before the sorted list of items
 	 * @param bLargeTimesFirst Whether to sort descending or ascending
+	 * @param bLogTotalTime Whether to log the total time taken when this log is destroyed
 	 * @param Units Whether to use seconds or milliseconds when printing time information
 	 * @param Formatter Receives the final data and logs it. Useful for logging in a special format, such as
 	 * @param OutputDevice Where to print to
 	 */
-	FConditionalSortedScopedLog(bool bCondition, FString Title = FString(), FLogFormatter Formatter = FFormatLog(), FOutputDevice* OutputDevice = GLog, bool bLargeTimesFirst = true, EScopeLogTimeUnits Units = EScopeLogTimeUnits::Milliseconds);
+	FConditionalSortedScopedLog(
+		bool bCondition,
+		FString Title = FString(),
+		FLogFormatter Formatter = FFormatLog(),
+		FOutputDevice* OutputDevice = GLog,
+		bool bLargeTimesFirst = true,
+		bool bLogTotalTime = true,
+		EScopeLogTimeUnits Units = EScopeLogTimeUnits::Milliseconds);
 	~FConditionalSortedScopedLog();
 
 	
@@ -82,6 +90,7 @@ private:
 	
 	FString LogTitle;
 	bool bLargeTimesFirst;
+	const double StartTime; 
 	EScopeLogTimeUnits Units;
 	FLogFormatter Formatter;
 	FOutputDevice* Output;
@@ -116,7 +125,13 @@ class FSortedScopedLog : public FConditionalSortedScopedLog
 {
 public:
 
-	FSortedScopedLog(FString Name = FString(), FLogFormatter Formatter = FFormatLog(), FOutputDevice* OutputDevice = GLog, bool bLargeTimesFirst = true, EScopeLogTimeUnits Units = EScopeLogTimeUnits::Milliseconds)
-		: FConditionalSortedScopedLog(true, Name, Formatter, OutputDevice, bLargeTimesFirst, Units)
+	FSortedScopedLog(FString Name = FString(),
+		FLogFormatter Formatter = FFormatLog(),
+		FOutputDevice* OutputDevice = GLog,
+		bool bLargeTimesFirst = true, 
+		bool bLogTotalTime = true,
+		EScopeLogTimeUnits Units = EScopeLogTimeUnits::Milliseconds
+		)
+		: FConditionalSortedScopedLog(true, Name, Formatter, OutputDevice, bLargeTimesFirst, bLogTotalTime, Units)
 	{}
 };
\ No newline at end of file
