Fix Foliage does not restore correctly after you move he static mesh asset

#[fyi] Jason.Walter
Fix static analysis warning

#[fyi] Jason.Walter

diff --git a/Engine/Plugins/VirtualProduction/LevelSnapshots/Source/FoliageSupport/Private/FoliageSupport/FoliageInfoData.cpp b/Engine/Plugins/VirtualProduction/LevelSnapshots/Source/FoliageSupport/Private/FoliageSupport/FoliageInfoData.cpp
--- a/Engine/Plugins/VirtualProduction/LevelSnapshots/Source/FoliageSupport/Private/FoliageSupport/FoliageInfoData.cpp
+++ b/Engine/Plugins/VirtualProduction/LevelSnapshots/Source/FoliageSupport/Private/FoliageSupport/FoliageInfoData.cpp
@@ -4,6 +4,7 @@
 
 #include "InstancedFoliageActor.h"
 #include "LevelSnapshotsLog.h"
+#include "SnapshotCustomVersion.h"
 #include "Serialization/MemoryReader.h"
 #include "Serialization/MemoryWriter.h"
 #include "Serialization/ObjectAndNameAsStringProxyArchive.h"
@@ -34,34 +35,46 @@ FArchive& UE::LevelSnapshots::Foliage::Private::FFoliageInfoData::SerializeInter
 {
 	Ar << ImplType;
 	Ar << ComponentName;
-	Ar << SerializedData;
+	if (Ar.CustomVer(FSnapshotCustomVersion::GUID) < FSnapshotCustomVersion::CustomSubobjectSoftObjectPathRefactor)
+	{
+		Ar << SerializedData_DEPRECATED;
+	}
+	else
+	{
+		Ar << FoliageInfoArchiveSize;
+	}
 	return Ar;
 }
 
-void UE::LevelSnapshots::Foliage::Private::FFoliageInfoData::Save(FFoliageInfo& DataToReadFrom, const FCustomVersionContainer& VersionInfo)
+void UE::LevelSnapshots::Foliage::Private::FFoliageInfoData::Save(FArchive& Archive, FFoliageInfo& DataToReadFrom)
 {
-	FMemoryWriter MemoryWriter(SerializedData, true);
-	FObjectAndNameAsStringProxyArchive RootArchive(MemoryWriter, false);
-	RootArchive.SetCustomVersions(VersionInfo);
-	
 	ImplType = DataToReadFrom.Type;
 	
 	const bool bIsComponentValid = DataToReadFrom.GetComponent() != nullptr; 
 	ComponentName = bIsComponentValid ? DataToReadFrom.GetComponent()->GetFName() : FName(NAME_None);
 	UE_CLOG(!bIsComponentValid, LogLevelSnapshots, Warning, TEXT("Could not save component for ImplType %s"), *AsName(ImplType));
 
-	RootArchive << DataToReadFrom;
+	int64 ArchivePos = Archive.Tell();
+	Archive << DataToReadFrom;
+	FoliageInfoArchiveSize = Archive.Tell() - ArchivePos;
 }
 
-void UE::LevelSnapshots::Foliage::Private::FFoliageInfoData::ApplyTo(FFoliageInfo& DataToWriteInto, const FCustomVersionContainer& VersionInfo) const
+void UE::LevelSnapshots::Foliage::Private::FFoliageInfoData::ApplyTo(FArchive& Archive, FFoliageInfo& DataToWriteInto) const
 {
-	FMemoryReader MemoryReader(SerializedData, true);
-	FObjectAndNameAsStringProxyArchive RootArchive(MemoryReader, false);
-	RootArchive.SetCustomVersions(VersionInfo);
-	
 	// Avoid foliage internal checks
 	DataToWriteInto.Implementation.Reset();
 	
-	RootArchive << DataToWriteInto;
+	if (Archive.CustomVer(FSnapshotCustomVersion::GUID) >= FSnapshotCustomVersion::CustomSubobjectSoftObjectPathRefactor)
+	{
+		Archive << DataToWriteInto;
+	}
+	else
+	{
+		FMemoryReader MemoryReader(SerializedData_DEPRECATED, true);
+		FObjectAndNameAsStringProxyArchive RootArchive(MemoryReader, false);
+		RootArchive.SetCustomVersions(Archive.GetCustomVersions());
+		RootArchive << DataToWriteInto;
+	}
+	
 	DataToWriteInto.RecomputeHash();
 }
\ No newline at end of file
diff --git a/Engine/Plugins/VirtualProduction/LevelSnapshots/Source/FoliageSupport/Private/FoliageSupport/FoliageInfoData.h b/Engine/Plugins/VirtualProduction/LevelSnapshots/Source/FoliageSupport/Private/FoliageSupport/FoliageInfoData.h
--- a/Engine/Plugins/VirtualProduction/LevelSnapshots/Source/FoliageSupport/Private/FoliageSupport/FoliageInfoData.h
+++ b/Engine/Plugins/VirtualProduction/LevelSnapshots/Source/FoliageSupport/Private/FoliageSupport/FoliageInfoData.h
@@ -12,18 +12,22 @@ namespace UE::LevelSnapshots::Foliage::Private
 	{
 		EFoliageImplType ImplType = EFoliageImplType::Unknown;
 		FName ComponentName;
-	
-		TArray<uint8> SerializedData;
+		
+		/** Used before FSnapshotCustomVersion::CustomSubobjectSoftObjectPathRefactor */
+		TArray<uint8> SerializedData_DEPRECATED;
+		/** Added after FSnapshotCustomVersion::CustomSubobjectSoftObjectPathRefactor */
+		int64 FoliageInfoArchiveSize = 0;
 
 		FArchive& SerializeInternal(FArchive& Ar);
 	
 	public:
 
-		void Save(FFoliageInfo& DataToReadFrom, const FCustomVersionContainer& VersionInfo);
-		void ApplyTo(FFoliageInfo& DataToWriteInto, const FCustomVersionContainer& VersionInfo) const;
+		void Save(FArchive& Archive, FFoliageInfo& DataToReadFrom);
+		void ApplyTo(FArchive& Archive, FFoliageInfo& DataToWriteInto) const;
 
 		EFoliageImplType GetImplType() const { return ImplType; }
 		TOptional<FName> GetComponentName() const { return ImplType == EFoliageImplType::StaticMesh && !ComponentName.IsNone() ? TOptional<FName>(ComponentName) : TOptional<FName>(); }
+		int64 GetFoliageInfoArchiveSize() const { return FoliageInfoArchiveSize; }
 	
 		friend FArchive& operator<<(FArchive& Ar, FFoliageInfoData& MeshInfo)
 		{
diff --git a/Engine/Plugins/VirtualProduction/LevelSnapshots/Source/FoliageSupport/Private/FoliageSupport/FoliageSupport.cpp b/Engine/Plugins/VirtualProduction/LevelSnapshots/Source/FoliageSupport/Private/FoliageSupport/FoliageSupport.cpp
--- a/Engine/Plugins/VirtualProduction/LevelSnapshots/Source/FoliageSupport/Private/FoliageSupport/FoliageSupport.cpp
+++ b/Engine/Plugins/VirtualProduction/LevelSnapshots/Source/FoliageSupport/Private/FoliageSupport/FoliageSupport.cpp
@@ -12,6 +12,7 @@
 #include "FoliageHelper.h"
 #include "ILevelSnapshotsModule.h"
 #include "InstancedFoliageActor.h"
+#include "SnapshotCustomVersion.h"
 #include "ActorPartition/ActorPartitionSubsystem.h"
 #if WITH_EDITOR
 #include "FoliageEditModule.h"
@@ -77,30 +78,13 @@ void UE::LevelSnapshots::Foliage::Private::FFoliageSupport::OnTakeSnapshot(UObje
 	AInstancedFoliageActor* FoliageActor = Cast<AInstancedFoliageActor>(EditorObject);
 	check(FoliageActor);
 
-	FInstancedFoliageActorData FoliageData;
-	DataStorage.WriteObjectAnnotation(FObjectAnnotator::CreateLambda([&FoliageData, FoliageActor](FArchive& Archive)
+	DataStorage.WriteObjectAnnotation(FObjectAnnotator::CreateLambda([FoliageActor](FArchive& Archive)
 	{
-		FoliageData.Save(Archive.GetCustomVersions(), FoliageActor);
-		Archive << FoliageData;
+		FInstancedFoliageActorData FoliageData;
+		FoliageData.Save(Archive, FoliageActor);
 	}));
 }
 
-void UE::LevelSnapshots::Foliage::Private::FFoliageSupport::PostApplySnapshotProperties(UObject* Object, const ICustomSnapshotSerializationData& DataStorage)
-{
-	AInstancedFoliageActor* FoliageActor = Cast<AInstancedFoliageActor>(Object);
-	check(FoliageActor);
-	
-	// Track this actor for safety so we know for sure that the functions were called in the order we expected them to
-	CurrentFoliageActor = FoliageActor;
-	DataStorage.ReadObjectAnnotation(FObjectAnnotator::CreateLambda([this](FArchive& Archive)
-	{
-		// "Archive" uses versioning information stored in the snapshot: it can just be reused 
-		CurrentVersionInfo = Archive.GetCustomVersions();
-		Archive << CurrentFoliageData;
-	}));
-
-	// Rest is done in PostApplySnapshotToActor (need access to the property selection map)
-}
 
 namespace UE::LevelSnapshots::FoliageSupport::Internal
 {
@@ -145,6 +129,36 @@ namespace UE::LevelSnapshots::FoliageSupport::Internal
 	}
 }
 
+void UE::LevelSnapshots::Foliage::Private::FFoliageSupport::PostApplyToEditorObject(UObject* Object, const ICustomSnapshotSerializationData& DataStorage, const FPropertySelectionMap& SelectionMap)
+{
+	AInstancedFoliageActor* FoliageActor = Cast<AInstancedFoliageActor>(Object);
+	check(FoliageActor);
+	
+	// Track this actor for safety so we know for sure that the functions were called in the order we expected them to
+	CurrentFoliageActor = FoliageActor;
+	DataStorage.ReadObjectAnnotation(FObjectAnnotator::CreateLambda([this, &SelectionMap, FoliageActor](FArchive& Archive)
+	{
+		FInstancedFoliageActorData CurrentFoliageData;
+		if (Archive.CustomVer(FSnapshotCustomVersion::GUID) < FSnapshotCustomVersion::CustomSubobjectSoftObjectPathRefactor)
+		{
+			Archive << CurrentFoliageData;
+		}
+		
+#if WITH_EDITOR
+		// Actor might not have been modified because none of its uproperties were changed
+		FoliageActor->Modify();
+#endif
+
+		const bool bWasRecreated = SelectionMap.GetDeletedActorsToRespawn().Contains(FoliageActor);
+		CurrentFoliageData.ApplyTo(Archive, FoliageActor, SelectionMap, bWasRecreated);
+		FoliageSupport::Internal::RebuildChangedFoliageComponents(FoliageActor, SelectionMap, bWasRecreated);
+
+		FoliageSupport::Internal::UpdateFoliageUI();
+	}));
+
+	// Rest is done in PostApplySnapshotToActor (need access to the property selection map)
+}
+
 void UE::LevelSnapshots::Foliage::Private::FFoliageSupport::PostApplySnapshotToActor(const FApplySnapshotToActorParams& Params)
 {
 	// This is the order of operations up until now
@@ -158,16 +172,6 @@ void UE::LevelSnapshots::Foliage::Private::FFoliageSupport::PostApplySnapshotToA
 	{
 		checkf(FoliageActor == CurrentFoliageActor.Get(), TEXT("PostApplySnapshotToActor was not directly followed by PostApplySnapshotProperties. Investigate."));
 		CurrentFoliageActor.Reset();
-#if WITH_EDITOR
-		// Actor might not have been modified because none of its uproperties were changed
-		FoliageActor->Modify();
-#endif
-
-		CurrentFoliageData.ApplyTo(CurrentVersionInfo, FoliageActor, Params.SelectedProperties, Params.bWasRecreated);
-		UE::LevelSnapshots::FoliageSupport::Internal::RebuildChangedFoliageComponents(FoliageActor, Params.SelectedProperties, Params.bWasRecreated);
-
-		FoliageTypesToRemove.Reset();
-		UE::LevelSnapshots::FoliageSupport::Internal::UpdateFoliageUI();
 	}
 }
 
diff --git a/Engine/Plugins/VirtualProduction/LevelSnapshots/Source/FoliageSupport/Private/FoliageSupport/FoliageSupport.h b/Engine/Plugins/VirtualProduction/LevelSnapshots/Source/FoliageSupport/Private/FoliageSupport/FoliageSupport.h
--- a/Engine/Plugins/VirtualProduction/LevelSnapshots/Source/FoliageSupport/Private/FoliageSupport/FoliageSupport.h
+++ b/Engine/Plugins/VirtualProduction/LevelSnapshots/Source/FoliageSupport/Private/FoliageSupport/FoliageSupport.h
@@ -21,17 +21,10 @@ namespace UE::LevelSnapshots::Foliage::Private
 	{
 		/** Just used to verify snapshots is in a valid state */
 		TWeakObjectPtr<AInstancedFoliageActor> CurrentFoliageActor;
-
-		/** Holds the verionn info of the last serialized foliage actor. Set by PostApplySnapshotProperties. */
-		FCustomVersionContainer CurrentVersionInfo;
-		/** Holds the foliage data serialized foliage actor. Set by PostApplySnapshotProperties. */
-		FInstancedFoliageActorData CurrentFoliageData;
-
-		TArray<UFoliageType*> FoliageTypesToRemove;
 		
 	public:
 
-		using ICustomObjectSnapshotSerializer::PreApplySnapshotProperties; 
+		using ICustomObjectSnapshotSerializer::PreApplyToSnapshotObject; 
 		
 		static void Register(ILevelSnapshotsModule& Module);
 
@@ -44,7 +37,7 @@ namespace UE::LevelSnapshots::Foliage::Private
 		virtual UObject* FindOrRecreateSubobjectInSnapshotWorld(UObject* SnapshotObject, const ISnapshotSubobjectMetaData& ObjectData, const ICustomSnapshotSerializationData& DataStorage) override { checkNoEntry(); return nullptr; }
 		virtual UObject* FindOrRecreateSubobjectInEditorWorld(UObject* EditorObject, const ISnapshotSubobjectMetaData& ObjectData, const ICustomSnapshotSerializationData& DataStorage) override { checkNoEntry(); return nullptr; }
 		virtual UObject* FindSubobjectInEditorWorld(UObject* EditorObject, const ISnapshotSubobjectMetaData& ObjectData, const ICustomSnapshotSerializationData& DataStorage) override { checkNoEntry(); return nullptr; }
-		virtual void PostApplySnapshotProperties(UObject* Object, const ICustomSnapshotSerializationData& DataStorage) override;
+		virtual void PostApplyToEditorObject(UObject* Object, const ICustomSnapshotSerializationData& DataStorage, const FPropertySelectionMap& SelectionMap) override;
 		//~ End ICustomObjectSnapshotSerializer Interface
 		
 		//~ Begin IRestorationListener Interface
diff --git a/Engine/Plugins/VirtualProduction/LevelSnapshots/Source/FoliageSupport/Private/FoliageSupport/InstancedFoliageActorData.cpp b/Engine/Plugins/VirtualProduction/LevelSnapshots/Source/FoliageSupport/Private/FoliageSupport/InstancedFoliageActorData.cpp
--- a/Engine/Plugins/VirtualProduction/LevelSnapshots/Source/FoliageSupport/Private/FoliageSupport/InstancedFoliageActorData.cpp
+++ b/Engine/Plugins/VirtualProduction/LevelSnapshots/Source/FoliageSupport/Private/FoliageSupport/InstancedFoliageActorData.cpp
@@ -2,6 +2,7 @@
 
 #include "FoliageSupport/InstancedFoliageActorData.h"
 
+#include "FoliageType.h"
 #include "FoliageRestorationInfo.h"
 #include "LevelSnapshotsLog.h"
 #include "Selection/PropertySelectionMap.h"
@@ -10,15 +11,16 @@
 
 namespace UE::LevelSnapshots::Foliage::Private::Internal
 {
-	static void SaveAsset(UFoliageType* FoliageType, FFoliageInfo& FoliageInfo, TMap<TSoftObjectPtr<UFoliageType>, FFoliageInfoData>& FoliageAssets, const FCustomVersionContainer& VersionInfo)
+	static void SaveAsset(FArchive& Archive, UFoliageType* FoliageType, FFoliageInfo& FoliageInfo, TMap<TSoftObjectPtr<UFoliageType>, FFoliageInfoData>& FoliageAssets)
 	{
-		FoliageAssets.Add(FoliageType).Save(FoliageInfo, VersionInfo);
+		FFoliageInfoData& InfoData = FoliageAssets.Add(FoliageType);
+		InfoData.Save(Archive, FoliageInfo);
 	}
 	
-	static void SaveSubobject(UFoliageType* FoliageType, FFoliageInfo& FoliageInfo,TArray<FSubobjectFoliageInfoData>& SubobjectData, const FCustomVersionContainer& VersionInfo)
+	static void SaveSubobject(FArchive& Archive, UFoliageType* FoliageType, FFoliageInfo& FoliageInfo,TArray<FSubobjectFoliageInfoData>& SubobjectData)
 	{
 		FSubobjectFoliageInfoData Data;
-		Data.Save(FoliageType, FoliageInfo, VersionInfo);
+		Data.Save(Archive, FoliageType, FoliageInfo);
 		SubobjectData.Emplace(Data);
 	}
 }
@@ -31,9 +33,9 @@ FArchive& UE::LevelSnapshots::Foliage::Private::FInstancedFoliageActorData::Seri
 	return Ar;
 }
 
-void UE::LevelSnapshots::Foliage::Private::FInstancedFoliageActorData::Save(const FCustomVersionContainer& VersionInfo, AInstancedFoliageActor* FoliageActor)
+void UE::LevelSnapshots::Foliage::Private::FInstancedFoliageActorData::Save(FArchive& Archive, AInstancedFoliageActor* FoliageActor)
 {
-	FoliageActor->ForEachFoliageInfo([this, &VersionInfo, FoliageActor](UFoliageType* FoliageType, FFoliageInfo& FoliageInfo)
+	FoliageActor->ForEachFoliageInfo([this, &Archive, FoliageActor](UFoliageType* FoliageType, FFoliageInfo& FoliageInfo)
 	{
 		FFoliageImpl* Impl = FoliageInfo.Implementation.Get();
 		if (!FoliageType || !Impl)
@@ -44,11 +46,11 @@ void UE::LevelSnapshots::Foliage::Private::FInstancedFoliageActorData::Save(cons
 		const bool bIsSubobject = FoliageType->IsIn(FoliageActor);
 		if (bIsSubobject)
 		{
-			Internal::SaveSubobject(FoliageType, FoliageInfo, SubobjectData, VersionInfo);
+			Internal::SaveSubobject(Archive, FoliageType, FoliageInfo, SubobjectData);
 		}
 		else
 		{
-			Internal::SaveAsset(FoliageType, FoliageInfo, FoliageAssets, VersionInfo);
+			Internal::SaveAsset(Archive, FoliageType, FoliageInfo, FoliageAssets);
 		}
 		
 		return true;
@@ -103,56 +105,93 @@ namespace UE::LevelSnapshots::Foliage::Private::Internal
 	}
 	
 	static void ApplyAssets(
-		const FCustomVersionContainer& VersionInfo,
+		FArchive& Archive,
 		AInstancedFoliageActor* FoliageActor,
 		const FFoliageRestorationInfo& RestorationInfo,
 		const TMap<TSoftObjectPtr<UFoliageType>, FFoliageInfoData>& FoliageAssets,
 		const TMap<FName, UFoliageType*> PreexistingComponentToFoliageType)
 	{
+		// Some objects may not want to be restored (either due to errors or because user deselected them) ...
+		int64 DataToSkip = 0;
+		
 		for (auto AssetIt = FoliageAssets.CreateConstIterator(); AssetIt; ++AssetIt)
 		{
 			const FFoliageInfoData& FoliageInfoData = AssetIt->Value;
+			
+			bool bWasRestored = false;
+			ON_SCOPE_EXIT
+			{
+				if (!bWasRestored)
+				{
+					// Standard component restoration has already recreated the component. Remove it again.
+					RemoveRecreatedComponent(FoliageActor, FoliageInfoData.GetComponentName());
+					DataToSkip += FoliageInfoData.GetFoliageInfoArchiveSize();
+					// TODO: Remove the foliage info from FoliageActor ... it was already serialized into FoliageInfos
+				}
+			};
+			
 			UFoliageType* FoliageType = AssetIt->Key.LoadSynchronous();
 			if (!FoliageType)
 			{
 				UE_LOG(LogLevelSnapshots, Warning, TEXT("Foliage type %s is missing. Component %s will not be restored."), *AssetIt->Key.ToString(), *FoliageInfoData.GetComponentName().Get(NAME_None).ToString());
-
-				// Standard component restoration has already recreated the component. Remove it again.
-				RemoveRecreatedComponent(FoliageActor, FoliageInfoData.GetComponentName());
 				continue;
 			}
 
-			if (RestorationInfo.ShouldSkipFoliageType(FoliageInfoData))
+			if (RestorationInfo.ShouldSkipFoliageType(FoliageInfoData) || !RestorationInfo.ShouldSerializeFoliageType(FoliageInfoData))
 			{
 				continue;
 			}
 
-			if (RestorationInfo.ShouldSerializeFoliageType(FoliageInfoData))
+			FFoliageInfo* FoliageInfo = FindOrAddFoliageInfo(FoliageType, FoliageActor);
+			if (FoliageInfo)
 			{
-				FFoliageInfo* FoliageInfo = FindOrAddFoliageInfo(FoliageType, FoliageActor);
-				if (FoliageInfo)
+				// If there were no instances, no component existed, hence no component name could be saved
+				const bool bHadAtLeastOneInstanceWhenSaved = FoliageInfoData.GetComponentName().IsSet();
+				if (bHadAtLeastOneInstanceWhenSaved)
 				{
-					// If there were no instances, no component existed, hence no component name could be saved
-					const bool bHadAtLeastOneInstanceWhenSaved = FoliageInfoData.GetComponentName().IsSet();
-					if (bHadAtLeastOneInstanceWhenSaved)
-					{
-						HandleExistingFoliageUsingRequiredComponent(FoliageActor, *FoliageInfoData.GetComponentName(), PreexistingComponentToFoliageType, FoliageType);
-					}
-					FoliageInfoData.ApplyTo(*FoliageInfo, VersionInfo);
+					HandleExistingFoliageUsingRequiredComponent(FoliageActor, *FoliageInfoData.GetComponentName(), PreexistingComponentToFoliageType, FoliageType);
 				}
+
+				// ... since Archive data is sequential we must account for possibly having skipped some data
+				Archive.Seek(Archive.Tell() + DataToSkip);
+				DataToSkip = 0;
+				
+				FoliageInfoData.ApplyTo(Archive, *FoliageInfo);
+				bWasRestored = true;
 			}
 		}
 	}
 
 	static void ApplySubobjects(
-		const FCustomVersionContainer& VersionInfo,
+		FArchive& Archive,
 		AInstancedFoliageActor* FoliageActor,
 		const FFoliageRestorationInfo& RestorationInfo,
 		const TArray<FSubobjectFoliageInfoData>& SubobjectData, 
 		const TMap<FName, UFoliageType*> PreexistingComponentToFoliageType)
 	{
+		// Some objects may not want to be restored (either due to errors or because user deselected them) ...
+		int64 DataToSkip = 0;
+		
 		for (const FSubobjectFoliageInfoData& InstanceData : SubobjectData)
 		{
+			bool bWasRestored = false;
+			bool bWasPartiallyRestored = false;
+			ON_SCOPE_EXIT
+			{
+				if (!bWasRestored)
+				{
+					// Standard component restoration has already recreated the component. Remove it again.
+					RemoveRecreatedComponent(FoliageActor, InstanceData.GetComponentName());
+					DataToSkip += static_cast<uint64>(bWasPartiallyRestored) * InstanceData.GetFoliageTypeArchiveSize() + InstanceData.GetFoliageInfoArchiveSize();
+					// TODO: Remove the foliage info from FoliageActor ... it was already serialized into FoliageInfos
+				}
+			};
+			
+			if (RestorationInfo.ShouldSkipFoliageType(InstanceData) || !RestorationInfo.ShouldSerializeFoliageType(InstanceData))
+			{
+				continue;
+			}
+
 			UFoliageType* FoliageType = InstanceData.FindOrRecreateSubobject(FoliageActor);
 			if (!FoliageType)
 			{
@@ -160,22 +199,24 @@ namespace UE::LevelSnapshots::Foliage::Private::Internal
 				continue;
 			}
 
-			if (RestorationInfo.ShouldSkipFoliageType(InstanceData))
+			// ... since Archive data is sequential we must account for possibly having skipped some data
+			Archive.Seek(Archive.Tell() + DataToSkip);
+			DataToSkip = 0;
+
+			bWasPartiallyRestored = true;
+			if (!InstanceData.ApplyToFoliageType(Archive, FoliageType))
 			{
+				UE_LOG(LogLevelSnapshots, Warning, TEXT("Skipping restoration of foliage item due to errors."));
 				continue;
 			}
-
-			if (RestorationInfo.ShouldSerializeFoliageType(InstanceData))
+				
+			FFoliageInfo* FoliageInfo = FindOrAddFoliageInfo(FoliageType, FoliageActor);
+			if (FoliageInfo
+				&& ensureAlwaysMsgf(InstanceData.GetComponentName(), TEXT("ComponentName is supposed to have been saved. Investigate. Maybe you used an unsupported foliage type (only EFoliageImplType::StaticMesh is supported)?")))
 			{
-				FFoliageInfo* FoliageInfo = FindOrAddFoliageInfo(FoliageType, FoliageActor);
-				if (FoliageInfo
-					&& ensureAlwaysMsgf(InstanceData.GetComponentName(), TEXT("ComponentName is supposed to have been saved. Investigate. Maybe you used an unsupported foliage type (only EFoliageImplType::StaticMesh is supported)?")))
-				{
-					HandleExistingFoliageUsingRequiredComponent(FoliageActor, *InstanceData.GetComponentName(), PreexistingComponentToFoliageType, FoliageType);
-
-					InstanceData.ApplyTo(FoliageType, VersionInfo);
-					InstanceData.ApplyTo(*FoliageInfo, VersionInfo);
-				}
+				HandleExistingFoliageUsingRequiredComponent(FoliageActor, *InstanceData.GetComponentName(), PreexistingComponentToFoliageType, FoliageType);
+				InstanceData.ApplyToFoliageInfo(Archive, *FoliageInfo);
+				bWasRestored = true;
 			}
 		}
 	}
@@ -194,10 +235,10 @@ namespace UE::LevelSnapshots::Foliage::Private::Internal
 	}
 }
 
-void UE::LevelSnapshots::Foliage::Private::FInstancedFoliageActorData::ApplyTo(const FCustomVersionContainer& VersionInfo, AInstancedFoliageActor* FoliageActor, const FPropertySelectionMap& SelectedProperties, bool bWasRecreated) const
+void UE::LevelSnapshots::Foliage::Private::FInstancedFoliageActorData::ApplyTo(FArchive& Archive, AInstancedFoliageActor* FoliageActor, const FPropertySelectionMap& SelectedProperties, bool bWasRecreated) const
 {
 	const FFoliageRestorationInfo RestorationInfo = FFoliageRestorationInfo::From(FoliageActor, SelectedProperties, bWasRecreated);
 	const TMap<FName, UFoliageType*> PreexistingComponentToFoliageType = Internal::BuildComponentToFoliageType(FoliageActor);
-	Internal::ApplyAssets(VersionInfo, FoliageActor, RestorationInfo, FoliageAssets, PreexistingComponentToFoliageType);
-	Internal::ApplySubobjects(VersionInfo, FoliageActor, RestorationInfo, SubobjectData, PreexistingComponentToFoliageType);
+	Internal::ApplyAssets(Archive, FoliageActor, RestorationInfo, FoliageAssets, PreexistingComponentToFoliageType);
+	Internal::ApplySubobjects(Archive, FoliageActor, RestorationInfo, SubobjectData, PreexistingComponentToFoliageType);
 }
diff --git a/Engine/Plugins/VirtualProduction/LevelSnapshots/Source/FoliageSupport/Private/FoliageSupport/InstancedFoliageActorData.h b/Engine/Plugins/VirtualProduction/LevelSnapshots/Source/FoliageSupport/Private/FoliageSupport/InstancedFoliageActorData.h
--- a/Engine/Plugins/VirtualProduction/LevelSnapshots/Source/FoliageSupport/Private/FoliageSupport/InstancedFoliageActorData.h
+++ b/Engine/Plugins/VirtualProduction/LevelSnapshots/Source/FoliageSupport/Private/FoliageSupport/InstancedFoliageActorData.h
@@ -35,12 +35,13 @@ namespace UE::LevelSnapshots::Foliage::Private
 	public:
 
 		/** Saves foliage actor into this object */
-		void Save(const FCustomVersionContainer& VersionInfo, AInstancedFoliageActor* FoliageActor);
+		void Save(FArchive& Archive, AInstancedFoliageActor* FoliageActor);
+		
 		/**
 		* Looks up saved foliage types and restores the passed in foliage actor.
 		* The caller must ensure that the required instances already exist by using FindFoliageType or FindOrCreateFoliageType.
 		*/
-		void ApplyTo(const FCustomVersionContainer& VersionInfo, AInstancedFoliageActor* FoliageActor, const FPropertySelectionMap& SelectedProperties, bool bWasRecreated) const;
+		void ApplyTo(FArchive& Archive, AInstancedFoliageActor* FoliageActor, const FPropertySelectionMap& SelectedProperties, bool bWasRecreated) const;
 		
 		friend FArchive& operator<<(FArchive& Ar, FInstancedFoliageActorData& MeshInfo)
 		{
diff --git a/Engine/Plugins/VirtualProduction/LevelSnapshots/Source/FoliageSupport/Private/FoliageSupport/SubobjectFoliageInfoData.cpp b/Engine/Plugins/VirtualProduction/LevelSnapshots/Source/FoliageSupport/Private/FoliageSupport/SubobjectFoliageInfoData.cpp
--- a/Engine/Plugins/VirtualProduction/LevelSnapshots/Source/FoliageSupport/Private/FoliageSupport/SubobjectFoliageInfoData.cpp
+++ b/Engine/Plugins/VirtualProduction/LevelSnapshots/Source/FoliageSupport/Private/FoliageSupport/SubobjectFoliageInfoData.cpp
@@ -4,9 +4,10 @@
 
 #include "LevelSnapshotsLog.h"
 
+#include "FoliageType_InstancedStaticMesh.h"
 #include "InstancedFoliageActor.h"
+#include "SnapshotCustomVersion.h"
 #include "Serialization/MemoryReader.h"
-#include "Serialization/MemoryWriter.h"
 #include "Serialization/ObjectAndNameAsStringProxyArchive.h"
 
 FArchive& UE::LevelSnapshots::Foliage::Private::FSubobjectFoliageInfoData::SerializeInternal(FArchive& Ar)
@@ -14,21 +15,27 @@ FArchive& UE::LevelSnapshots::Foliage::Private::FSubobjectFoliageInfoData::Seria
 	Ar << static_cast<FFoliageInfoData&>(*this);
 	Ar << Class;
 	Ar << SubobjectName;
-	Ar << SerializedSubobjectData;
+	if (Ar.CustomVer(FSnapshotCustomVersion::GUID) < FSnapshotCustomVersion::CustomSubobjectSoftObjectPathRefactor)
+	{
+		Ar << SerializedSubobjectData_DEPRECATED;
+	}
+	else
+	{
+		Ar << FoliageTypeArchiveSize;
+	}
 	return Ar;
 }
 
-void UE::LevelSnapshots::Foliage::Private::FSubobjectFoliageInfoData::Save(UFoliageType* FoliageSubobject, FFoliageInfo& FoliageInfo, const FCustomVersionContainer& VersionInfo)
+void UE::LevelSnapshots::Foliage::Private::FSubobjectFoliageInfoData::Save(FArchive& Archive, UFoliageType* FoliageSubobject, FFoliageInfo& FoliageInfo)
 {
-	FFoliageInfoData::Save(FoliageInfo, VersionInfo);
+	FFoliageInfoData::Save(Archive, FoliageInfo);
 
 	Class = FoliageSubobject->GetClass();
 	SubobjectName = FoliageSubobject->GetFName();
-	
-	FMemoryWriter MemoryWriter(SerializedSubobjectData, true);
-	FObjectAndNameAsStringProxyArchive RootArchive(MemoryWriter, false);
-	RootArchive.SetCustomVersions(VersionInfo);
-	FoliageSubobject->Serialize(RootArchive);
+
+	const int64 ArchivePos = Archive.Tell();
+	FoliageSubobject->Serialize(Archive);
+	FoliageTypeArchiveSize += Archive.Tell() - ArchivePos;
 }
 
 UFoliageType* UE::LevelSnapshots::Foliage::Private::FSubobjectFoliageInfoData::FindOrRecreateSubobject(AInstancedFoliageActor* Outer) const
@@ -39,18 +46,35 @@ UFoliageType* UE::LevelSnapshots::Foliage::Private::FSubobjectFoliageInfoData::F
 		return Cast<UFoliageType>(FoundObject);
 	}
 
-	return NewObject<UFoliageType>(Outer, Class, SubobjectName);
+	return NewObject<UFoliageType>(Outer, Class, SubobjectName, RF_Transactional);
 }
 
-void UE::LevelSnapshots::Foliage::Private::FSubobjectFoliageInfoData::ApplyTo(UFoliageType* FoliageSubobject, const FCustomVersionContainer& VersionInfo) const
+bool UE::LevelSnapshots::Foliage::Private::FSubobjectFoliageInfoData::ApplyToFoliageType(FArchive& Archive, UFoliageType* FoliageSubobject) const
 {
-	FMemoryReader MemoryReader(SerializedSubobjectData, true);
-	constexpr bool bLoadIfFindFails = true;
-	FObjectAndNameAsStringProxyArchive RootArchive(MemoryReader, bLoadIfFindFails);
-	FoliageSubobject->Serialize(RootArchive);
+	if (Archive.CustomVer(FSnapshotCustomVersion::GUID) >= FSnapshotCustomVersion::CustomSubobjectSoftObjectPathRefactor)
+	{
+		FoliageSubobject->Serialize(Archive);
+	}
+	else
+	{
+		FMemoryReader MemoryReader(SerializedSubobjectData_DEPRECATED, true);
+		constexpr bool bLoadIfFindFails = true;
+		FObjectAndNameAsStringProxyArchive RootArchive(MemoryReader, bLoadIfFindFails);
+		FoliageSubobject->Serialize(RootArchive);
+	}
+
+	if (UFoliageType_InstancedStaticMesh* MeshFoliageType = Cast<UFoliageType_InstancedStaticMesh>(FoliageSubobject)
+		; ensureMsgf(MeshFoliageType, TEXT("Only static mesh foliage is supported right now")))
+	{
+		// There used to be a bug where the asset registry was not aware 
+		UE_LOG(LogLevelSnapshots, Error, TEXT("Reference to foliage mesh was corrupted."));
+		return MeshFoliageType->Mesh != nullptr;
+	}
+
+	return true;
 }
 
-void UE::LevelSnapshots::Foliage::Private::FSubobjectFoliageInfoData::ApplyTo(FFoliageInfo& DataToWriteInto, const FCustomVersionContainer& VersionInfo) const
+void UE::LevelSnapshots::Foliage::Private::FSubobjectFoliageInfoData::ApplyToFoliageInfo(FArchive& Archive, FFoliageInfo& DataToWriteInto) const
 {
-	FFoliageInfoData::ApplyTo(DataToWriteInto, VersionInfo);
+	FFoliageInfoData::ApplyTo(Archive, DataToWriteInto);
 }
diff --git a/Engine/Plugins/VirtualProduction/LevelSnapshots/Source/FoliageSupport/Private/FoliageSupport/SubobjectFoliageInfoData.h b/Engine/Plugins/VirtualProduction/LevelSnapshots/Source/FoliageSupport/Private/FoliageSupport/SubobjectFoliageInfoData.h
--- a/Engine/Plugins/VirtualProduction/LevelSnapshots/Source/FoliageSupport/Private/FoliageSupport/SubobjectFoliageInfoData.h
+++ b/Engine/Plugins/VirtualProduction/LevelSnapshots/Source/FoliageSupport/Private/FoliageSupport/SubobjectFoliageInfoData.h
@@ -17,17 +17,23 @@ namespace UE::LevelSnapshots::Foliage::Private
 	{
 		TSubclassOf<UFoliageType> Class;
 		FName SubobjectName;
-		TArray<uint8> SerializedSubobjectData;
+		
+		/** Used before FSnapshotCustomVersion::CustomSubobjectSoftObjectPathRefactor*/
+		TArray<uint8> SerializedSubobjectData_DEPRECATED;
+		/** Added after FSnapshotCustomVersion::CustomSubobjectSoftObjectPathRefactor */
+		int64 FoliageTypeArchiveSize = 0;
 
 		FArchive& SerializeInternal(FArchive& Ar);
 	
 	public:
 
-		void Save(UFoliageType* FoliageSubobject, FFoliageInfo& FoliageInfo, const FCustomVersionContainer& VersionInfo);
+		void Save(FArchive& Archive, UFoliageType* FoliageSubobject, FFoliageInfo& FoliageInfo);
 	
 		UFoliageType* FindOrRecreateSubobject(AInstancedFoliageActor* Outer) const;
-		void ApplyTo(UFoliageType* FoliageSubobject, const FCustomVersionContainer& VersionInfo) const;
-		void ApplyTo(FFoliageInfo& DataToWriteInto, const FCustomVersionContainer& VersionInfo) const;
+		bool ApplyToFoliageType(FArchive& Archive, UFoliageType* FoliageSubobject) const;
+		void ApplyToFoliageInfo(FArchive& Archive, FFoliageInfo& DataToWriteInto) const;
+
+		int64 GetFoliageTypeArchiveSize() const { return FoliageTypeArchiveSize; }
 		
 		friend FArchive& operator<<(FArchive& Ar, FSubobjectFoliageInfoData& Data)
 		{
diff --git a/Engine/Plugins/VirtualProduction/LevelSnapshots/Source/LevelSnapshots/Private/Archive/CustomSerialization/BaseCustomObjectProxyArchive.cpp b/Engine/Plugins/VirtualProduction/LevelSnapshots/Source/LevelSnapshots/Private/Archive/CustomSerialization/BaseCustomObjectProxyArchive.cpp
new file mode 100644
--- /dev/null
+++ b/Engine/Plugins/VirtualProduction/LevelSnapshots/Source/LevelSnapshots/Private/Archive/CustomSerialization/BaseCustomObjectProxyArchive.cpp
@@ -0,0 +1,29 @@
+// Copyright Epic Games, Inc. All Rights Reserved.
+
+#include "BaseCustomObjectProxyArchive.h"
+
+#include "Serialization/ArchiveUObject.h"
+#include "UObject/SoftObjectPtr.h"
+
+namespace UE::LevelSnapshots::CustomSerialization
+{
+	FArchive& FBaseCustomObjectProxyArchive::operator<<(FWeakObjectPtr& Obj)
+	{
+		return FArchiveUObject::SerializeWeakObjectPtr(*this, Obj);
+	}
+		
+	FArchive& FBaseCustomObjectProxyArchive::operator<<(FSoftObjectPtr& Value)
+	{
+		if (IsLoading())
+		{
+			Value.ResetWeakPtr();
+		}
+		*this << Value.GetUniqueID();
+		return *this;
+	}
+		
+	FArchive& FBaseCustomObjectProxyArchive::operator<<(FObjectPtr& Obj)
+	{
+		return FArchiveUObject::SerializeObjectPtr(*this, Obj);
+	}
+}
diff --git a/Engine/Plugins/VirtualProduction/LevelSnapshots/Source/LevelSnapshots/Private/Archive/CustomSerialization/BaseCustomObjectProxyArchive.h b/Engine/Plugins/VirtualProduction/LevelSnapshots/Source/LevelSnapshots/Private/Archive/CustomSerialization/BaseCustomObjectProxyArchive.h
new file mode 100644
--- /dev/null
+++ b/Engine/Plugins/VirtualProduction/LevelSnapshots/Source/LevelSnapshots/Private/Archive/CustomSerialization/BaseCustomObjectProxyArchive.h
@@ -0,0 +1,26 @@
+// Copyright Epic Games, Inc. All Rights Reserved.
+
+#pragma once
+
+#include "CoreMinimal.h"
+#include "Serialization/NameAsStringProxyArchive.h"
+
+namespace UE::LevelSnapshots::CustomSerialization
+{
+	/** Primarily used for serializing object annotation data for implementing ICustomSnapshotSerializationData. */
+	class FBaseCustomObjectProxyArchive : public FNameAsStringProxyArchive
+	{
+	public:
+		
+		FBaseCustomObjectProxyArchive(FArchive& InInnerArchive)
+			: FNameAsStringProxyArchive(InInnerArchive)
+		{}
+
+		virtual FArchive& operator<<(UObject*& Obj) override = 0;
+		virtual FArchive& operator<<(FSoftObjectPath& Value) override = 0;
+		
+		virtual FArchive& operator<<(FWeakObjectPtr& Obj) override;
+		virtual FArchive& operator<<(FSoftObjectPtr& Value) override;
+		virtual FArchive& operator<<(FObjectPtr& Obj) override;
+	};
+}
\ No newline at end of file
diff --git a/Engine/Plugins/VirtualProduction/LevelSnapshots/Source/LevelSnapshots/Private/Archive/CustomSerialization/LoadCustomObjectProxyArchive.cpp b/Engine/Plugins/VirtualProduction/LevelSnapshots/Source/LevelSnapshots/Private/Archive/CustomSerialization/LoadCustomObjectProxyArchive.cpp
new file mode 100644
--- /dev/null
+++ b/Engine/Plugins/VirtualProduction/LevelSnapshots/Source/LevelSnapshots/Private/Archive/CustomSerialization/LoadCustomObjectProxyArchive.cpp
@@ -0,0 +1,31 @@
+// Copyright Epic Games, Inc. All Rights Reserved.
+
+#include "LoadCustomObjectProxyArchive.h"
+
+#include "UObject/SoftObjectPath.h"
+
+namespace UE::LevelSnapshots::CustomSerialization
+{
+	FLoadCustomObjectProxyArchive::FLoadCustomObjectProxyArchive(FArchive& InInnerArchive, const TArray<FSoftObjectPath>& TargetArray)
+			: FBaseCustomObjectProxyArchive(InInnerArchive)
+			, TargetArray(TargetArray)
+	{
+		check(InInnerArchive.IsLoading());
+	}
+
+	FArchive& FLoadCustomObjectProxyArchive::operator<<(UObject*& Obj)
+	{
+		FSoftObjectPath Path;
+		*this << Path;
+		Obj = Path.TryLoad();
+		return *this;
+	}
+		
+	FArchive& FLoadCustomObjectProxyArchive::operator<<(FSoftObjectPath& Value)
+	{
+		int32 ReferenceIndex;
+		*this << ReferenceIndex;
+		Value = TargetArray[ReferenceIndex];
+		return *this;
+	}
+}
diff --git a/Engine/Plugins/VirtualProduction/LevelSnapshots/Source/LevelSnapshots/Private/Archive/CustomSerialization/LoadCustomObjectProxyArchive.h b/Engine/Plugins/VirtualProduction/LevelSnapshots/Source/LevelSnapshots/Private/Archive/CustomSerialization/LoadCustomObjectProxyArchive.h
new file mode 100644
--- /dev/null
+++ b/Engine/Plugins/VirtualProduction/LevelSnapshots/Source/LevelSnapshots/Private/Archive/CustomSerialization/LoadCustomObjectProxyArchive.h
@@ -0,0 +1,20 @@
+// Copyright Epic Games, Inc. All Rights Reserved.
+
+#pragma once
+
+#include "CoreMinimal.h"
+#include "BaseCustomObjectProxyArchive.h"
+
+namespace UE::LevelSnapshots::CustomSerialization
+{
+	class FLoadCustomObjectProxyArchive : public FBaseCustomObjectProxyArchive
+	{
+		const TArray<FSoftObjectPath>& TargetArray;
+	public:
+
+		FLoadCustomObjectProxyArchive(FArchive& InInnerArchive, const TArray<FSoftObjectPath>& TargetArray);
+
+		virtual FArchive& operator<<(UObject*& Obj) override;
+		virtual FArchive& operator<<(FSoftObjectPath& Value) override;
+	};
+}
diff --git a/Engine/Plugins/VirtualProduction/LevelSnapshots/Source/LevelSnapshots/Private/Archive/CustomSerialization/SaveCustomObjectProxyArchive.cpp b/Engine/Plugins/VirtualProduction/LevelSnapshots/Source/LevelSnapshots/Private/Archive/CustomSerialization/SaveCustomObjectProxyArchive.cpp
new file mode 100644
--- /dev/null
+++ b/Engine/Plugins/VirtualProduction/LevelSnapshots/Source/LevelSnapshots/Private/Archive/CustomSerialization/SaveCustomObjectProxyArchive.cpp
@@ -0,0 +1,30 @@
+// Copyright Epic Games, Inc. All Rights Reserved.
+
+#include "SaveCustomObjectProxyArchive.h"
+
+#include "UObject/SoftObjectPath.h"
+
+namespace UE::LevelSnapshots::CustomSerialization
+{
+	FSaveCustomObjectProxyArchive::FSaveCustomObjectProxyArchive(FArchive& InInnerArchive, TArray<FSoftObjectPath>& TargetArray)
+			: FBaseCustomObjectProxyArchive(InInnerArchive)
+			, TargetArray(TargetArray)
+	{
+		check(InInnerArchive.IsSaving());
+	}
+
+	FArchive& FSaveCustomObjectProxyArchive::operator<<(UObject*& Obj)
+	{
+		// This could be a world reference... we ignore this case for now. It has an effect when we restore into the snapshot world.
+		FSoftObjectPath Path = Obj;
+		*this << Path;
+		return *this;
+	}
+		
+	FArchive& FSaveCustomObjectProxyArchive::operator<<(FSoftObjectPath& Value)
+	{
+		int32 Index = TargetArray.AddUnique(Value);
+		*this << Index;
+		return *this;
+	}
+}
\ No newline at end of file
diff --git a/Engine/Plugins/VirtualProduction/LevelSnapshots/Source/LevelSnapshots/Private/Archive/CustomSerialization/SaveCustomObjectProxyArchive.h b/Engine/Plugins/VirtualProduction/LevelSnapshots/Source/LevelSnapshots/Private/Archive/CustomSerialization/SaveCustomObjectProxyArchive.h
new file mode 100644
--- /dev/null
+++ b/Engine/Plugins/VirtualProduction/LevelSnapshots/Source/LevelSnapshots/Private/Archive/CustomSerialization/SaveCustomObjectProxyArchive.h
@@ -0,0 +1,21 @@
+// Copyright Epic Games, Inc. All Rights Reserved.
+
+#pragma once
+
+#include "CoreMinimal.h"
+#include "BaseCustomObjectProxyArchive.h"
+
+namespace UE::LevelSnapshots::CustomSerialization
+{
+	class FSaveCustomObjectProxyArchive : public FBaseCustomObjectProxyArchive
+	{
+		TArray<FSoftObjectPath>& TargetArray;
+	public:
+		
+		FSaveCustomObjectProxyArchive(FArchive& InInnerArchive, TArray<FSoftObjectPath>& TargetArray);
+		
+		virtual FArchive& operator<<(UObject*& Obj) override;
+		virtual FArchive& operator<<(FSoftObjectPath& Value) override;
+	};
+}
+
diff --git a/Engine/Plugins/VirtualProduction/LevelSnapshots/Source/LevelSnapshots/Private/Data/CustomSerialization/CustomObjectSerializationWrapper.cpp b/Engine/Plugins/VirtualProduction/LevelSnapshots/Source/LevelSnapshots/Private/Data/CustomSerialization/CustomObjectSerializationWrapper.cpp
--- a/Engine/Plugins/VirtualProduction/LevelSnapshots/Source/LevelSnapshots/Private/Data/CustomSerialization/CustomObjectSerializationWrapper.cpp
+++ b/Engine/Plugins/VirtualProduction/LevelSnapshots/Source/LevelSnapshots/Private/Data/CustomSerialization/CustomObjectSerializationWrapper.cpp
@@ -40,7 +40,7 @@ namespace UE::LevelSnapshots::Private::Internal
 		}
 
 		FCustomSerializationDataReader SerializationDataReader = FCustomSerializationDataReader(FCustomSerializationDataGetter_ReadOnly::CreateLambda([SerializationDataGetter](){ return SerializationDataGetter();}), WorldData);
-		CustomSerializer->PreApplySnapshotProperties(SnapshotObject, SerializationDataReader);
+		CustomSerializer->PreApplyToSnapshotObject(SnapshotObject, SerializationDataReader);
 		return FRestoreObjectScope([SnapshotObject, SerializationDataGetter, &WorldData, &Cache, &ProcessObjectDependency, LocalisationSnapshotPackage, SerializationDataReader, CustomSerializer]()
 		{
 			for (int32 i = 0; i < SerializationDataReader.GetNumSubobjects(); ++i)
@@ -60,7 +60,7 @@ namespace UE::LevelSnapshots::Private::Internal
 				CustomSerializer->OnPostSerializeSnapshotSubobject(SnapshotSubobject, *MetaData, SerializationDataReader);
 			}
 
-			CustomSerializer->PostApplySnapshotProperties(SnapshotObject, SerializationDataReader);
+			CustomSerializer->PostApplyToSnapshotObject(SnapshotObject, SerializationDataReader);
 		});	
 	}
 
@@ -84,7 +84,7 @@ namespace UE::LevelSnapshots::Private::Internal
 		}
 		
 		FCustomSerializationDataReader SerializationDataReader = FCustomSerializationDataReader(FCustomSerializationDataGetter_ReadOnly::CreateLambda([SerializationDataGetter](){ return SerializationDataGetter(); }), WorldData);
-		CustomSerializer->PreApplySnapshotProperties(EditorObject, SerializationDataReader);
+		CustomSerializer->PreApplyToEditorObject(EditorObject, SerializationDataReader, SelectionMap);
 		return FRestoreObjectScope([SnapshotObject, EditorObject, &WorldData, &Cache, &SelectionMap, LocalisationSnapshotPackage, SerializationDataGetter, SerializationDataReader, CustomSerializer]()
 		{
 			for (int32 i = 0; i < SerializationDataReader.GetNumSubobjects(); ++i)
@@ -139,7 +139,7 @@ namespace UE::LevelSnapshots::Private::Internal
 				UE_LOG(LogLevelSnapshots, Warning, TEXT("Editor subobject %s was not restored"), *EditorSubobject->GetPathName());	
 			}
 
-			CustomSerializer->PostApplySnapshotProperties(EditorObject, SerializationDataReader);
+			CustomSerializer->PostApplyToEditorObject(EditorObject, SerializationDataReader, SelectionMap);
 		});	
 	}
 }
diff --git a/Engine/Plugins/VirtualProduction/LevelSnapshots/Source/LevelSnapshots/Private/Data/CustomSerialization/CustomSerializationDataManager.cpp b/Engine/Plugins/VirtualProduction/LevelSnapshots/Source/LevelSnapshots/Private/Data/CustomSerialization/CustomSerializationDataManager.cpp
--- a/Engine/Plugins/VirtualProduction/LevelSnapshots/Source/LevelSnapshots/Private/Data/CustomSerialization/CustomSerializationDataManager.cpp
+++ b/Engine/Plugins/VirtualProduction/LevelSnapshots/Source/LevelSnapshots/Private/Data/CustomSerialization/CustomSerializationDataManager.cpp
@@ -4,75 +4,121 @@
 
 #include "CustomObjectSerializationWrapper.h"
 #include "LevelSnapshotsLog.h"
+#include "SnapshotCustomVersion.h"
 #include "TakeWorldObjectSnapshotArchive.h"
+#include "CustomSerialization/LoadCustomObjectProxyArchive.h"
+#include "CustomSerialization/SaveCustomObjectProxyArchive.h"
+
 #include "Serialization/MemoryWriter.h"
 #include "Serialization/MemoryReader.h"
-#include "Serialization/NameAsStringProxyArchive.h"
 #include "Serialization/ObjectAndNameAsStringProxyArchive.h"
 #include "Util/WorldData/SnapshotObjectUtil.h"
 
 namespace UE::LevelSnapshots::Private::Internal
 {
-	class FSnapshotSubobjectMetaDataManager
-	:
-	public ISnapshotSubobjectMetaData,
-	public TSharedFromThis<FSnapshotSubobjectMetaDataManager>
+	class FSnapshotSubobjectMetaDataReader
+		:
+		public ISnapshotSubobjectMetaData,
+		public TSharedFromThis<FSnapshotSubobjectMetaDataReader>
 	{
-	public:
-	
 		const FSoftObjectPath CachedSubobjectPath;
 		const int32 OwningDataSubobjectIndex;
 		const FWorldSnapshotData& WorldData;
 	
 		const FCustomSerializationDataGetter_ReadOnly SerializationDataGetter_ReadOnly;
-		const FCustomSerializationDataGetter_ReadWrite SerializationDataGetter_ReadWrite;
+	protected:
+		
+		int32 GetOwningDataSubobjectIndex() const { return OwningDataSubobjectIndex; }
+		
+	public:
 
-		FSnapshotSubobjectMetaDataManager(
+		FSnapshotSubobjectMetaDataReader(
 			FSoftObjectPath CachedSubobjectPath,
 			int32 SubobjectIndex,
 			const FWorldSnapshotData& WorldData,
-			const FCustomSerializationDataGetter_ReadOnly& SerializationDataGetter_ReadOnly,
-			FCustomSerializationDataGetter_ReadWrite SerializationDataGetter_ReadWrite)
-			:
-			CachedSubobjectPath(CachedSubobjectPath),
-			OwningDataSubobjectIndex(SubobjectIndex),
-			WorldData(WorldData),
-			SerializationDataGetter_ReadOnly(SerializationDataGetter_ReadOnly),
-			SerializationDataGetter_ReadWrite(SerializationDataGetter_ReadWrite)
+			FCustomSerializationDataGetter_ReadOnly SerializationDataGetter_ReadOnly
+			)
+			: CachedSubobjectPath(CachedSubobjectPath)
+			, OwningDataSubobjectIndex(SubobjectIndex)
+			, WorldData(WorldData)
+			, SerializationDataGetter_ReadOnly(SerializationDataGetter_ReadOnly)
 		{}
 
 		virtual FSoftObjectPath GetOriginalPath() const override
 		{
 			return CachedSubobjectPath;
 		}
-	
+		
 		virtual void WriteObjectAnnotation(const FObjectAnnotator& Writer) override
 		{
-			if (Writer.IsBound() && ensure(SerializationDataGetter_ReadWrite.IsBound()))
+			checkNoEntry(); // Not allowed on read-only object
+		}
+		
+		virtual void ReadObjectAnnotation(const FObjectAnnotator& Reader) const override
+		{
+			if (Reader.IsBound() && ensure(SerializationDataGetter_ReadOnly.IsBound()))
 			{
-				FMemoryWriter MemoryWriter(SerializationDataGetter_ReadWrite.Execute()->Subobjects[OwningDataSubobjectIndex].SubobjectAnnotationData, true);
-				FObjectAndNameAsStringProxyArchive RootArchive(MemoryWriter, false);
-				Writer.Execute(RootArchive);
+				if (WorldData.SnapshotVersionInfo.GetSnapshotCustomVersion() >= FSnapshotCustomVersion::CustomSubobjectSoftObjectPathRefactor)
+				{
+					FMemoryReader MemoryReader(SerializationDataGetter_ReadOnly.Execute()->Subobjects[OwningDataSubobjectIndex].SubobjectAnnotationData, true);
+					CustomSerialization::FLoadCustomObjectProxyArchive RootArchive(MemoryReader, WorldData.SerializedObjectReferences);
+					WorldData.SnapshotVersionInfo.ApplyToArchive(RootArchive);
+					Reader.Execute(RootArchive);
+				}
+				else
+				{
+					FMemoryReader MemoryReader(SerializationDataGetter_ReadOnly.Execute()->Subobjects[OwningDataSubobjectIndex].SubobjectAnnotationData, true);
+					FObjectAndNameAsStringProxyArchive RootArchive(MemoryReader, true);
+					WorldData.SnapshotVersionInfo.ApplyToArchive(RootArchive);
+					Reader.Execute(RootArchive);
+				}
 			}
 		}
+	};
+
+	class FSnapshotSubobjectMetaDataWriter : public FSnapshotSubobjectMetaDataReader
+	{
+		FWorldSnapshotData& WorldData_ReadWrite;
+		const FCustomSerializationDataGetter_ReadWrite SerializationDataGetter_ReadWrite;
+	public:
+
+		FSnapshotSubobjectMetaDataWriter(
+			FSoftObjectPath CachedSubobjectPath,
+			int32 SubobjectIndex,
+			FWorldSnapshotData& WorldData,
+			FCustomSerializationDataGetter_ReadOnly SerializationDataGetter_ReadOnly,
+			FCustomSerializationDataGetter_ReadWrite SerializationDataGetter_ReadWrite
+			)
+			: FSnapshotSubobjectMetaDataReader(CachedSubobjectPath, SubobjectIndex, WorldData, SerializationDataGetter_ReadOnly)
+			, WorldData_ReadWrite(WorldData)
+			, SerializationDataGetter_ReadWrite(SerializationDataGetter_ReadWrite)
+		{}
 	
-		virtual void ReadObjectAnnotation(const FObjectAnnotator& Reader) const override
+		virtual void WriteObjectAnnotation(const FObjectAnnotator& Writer) override
 		{
-			if (Reader.IsBound() && ensure(SerializationDataGetter_ReadOnly.IsBound()))
+			if (Writer.IsBound() && ensure(SerializationDataGetter_ReadWrite.IsBound()))
 			{
-				FMemoryReader MemoryReader(SerializationDataGetter_ReadOnly.Execute()->Subobjects[OwningDataSubobjectIndex].SubobjectAnnotationData, true);
-				FObjectAndNameAsStringProxyArchive RootArchive(MemoryReader, true);
-				WorldData.SnapshotVersionInfo.ApplyToArchive(RootArchive);
-				Reader.Execute(RootArchive);
+				if (WorldData_ReadWrite.SnapshotVersionInfo.GetSnapshotCustomVersion() >= FSnapshotCustomVersion::CustomSubobjectSoftObjectPathRefactor)
+				{
+					FMemoryWriter MemoryWriter(SerializationDataGetter_ReadWrite.Execute()->Subobjects[GetOwningDataSubobjectIndex()].SubobjectAnnotationData, true);
+					CustomSerialization::FSaveCustomObjectProxyArchive RootArchive(MemoryWriter, WorldData_ReadWrite.SerializedObjectReferences);
+					WorldData_ReadWrite.SnapshotVersionInfo.ApplyToArchive(RootArchive);
+					Writer.Execute(RootArchive);
+				}
+				else
+				{
+					FMemoryWriter MemoryWriter(SerializationDataGetter_ReadWrite.Execute()->Subobjects[GetOwningDataSubobjectIndex()].SubobjectAnnotationData, true);
+					FObjectAndNameAsStringProxyArchive RootArchive(MemoryWriter, false);
+					Writer.Execute(RootArchive);
+				}
 			}
 		}
 	};
 }
 
 UE::LevelSnapshots::Private::FCustomSerializationDataReader::FCustomSerializationDataReader(FCustomSerializationDataGetter_ReadOnly SerializationDataGetter, const FWorldSnapshotData& ConstWorldData)
-		:
-		SerializationDataGetter_ReadOnly(SerializationDataGetter),
-		WorldData_ReadOnly(ConstWorldData)
+	: SerializationDataGetter_ReadOnly(SerializationDataGetter)
+	, WorldData_ReadOnly(ConstWorldData)
 {
 	if (SerializationDataGetter.IsBound())
 	{
@@ -81,12 +127,11 @@ UE::LevelSnapshots::Private::FCustomSerializationDataReader::FCustomSerializatio
 		{
 			const FCustomSubbjectSerializationData& SubobjectData = SerializationData->Subobjects[i];
 			CachedSubobjectMetaData.Add(
-				MakeShared<Internal::FSnapshotSubobjectMetaDataManager>(
-					ConstWorldData.SerializedObjectReferences[SubobjectData.ObjectPathIndex],
+				MakeShared<Internal::FSnapshotSubobjectMetaDataReader>(
+						ConstWorldData.SerializedObjectReferences[SubobjectData.ObjectPathIndex],
 						i,
 						WorldData_ReadOnly,
-						SerializationDataGetter_ReadOnly,
-						FCustomSerializationDataGetter_ReadWrite()
+						SerializationDataGetter_ReadOnly
 					)
 				);
 		}
@@ -102,10 +147,20 @@ void UE::LevelSnapshots::Private::FCustomSerializationDataReader::ReadObjectAnno
 {
 	if (Reader.IsBound())
 	{
-		FMemoryReader MemoryReader(SerializationDataGetter_ReadOnly.Execute()->RootAnnotationData, true);
-		FObjectAndNameAsStringProxyArchive RootArchive(MemoryReader, true);
-		WorldData_ReadOnly.SnapshotVersionInfo.ApplyToArchive(RootArchive);
-		Reader.Execute(RootArchive);
+		if (WorldData_ReadOnly.SnapshotVersionInfo.GetSnapshotCustomVersion() >= FSnapshotCustomVersion::CustomSubobjectSoftObjectPathRefactor)
+		{
+			FMemoryReader MemoryReader(SerializationDataGetter_ReadOnly.Execute()->RootAnnotationData, true);
+			CustomSerialization::FLoadCustomObjectProxyArchive RootArchive(MemoryReader, WorldData_ReadOnly.SerializedObjectReferences);
+			WorldData_ReadOnly.SnapshotVersionInfo.ApplyToArchive(RootArchive);
+			Reader.Execute(RootArchive);
+		}
+		else
+		{
+			FMemoryReader MemoryReader(SerializationDataGetter_ReadOnly.Execute()->RootAnnotationData, true);
+			FObjectAndNameAsStringProxyArchive RootArchive(MemoryReader, true);
+			WorldData_ReadOnly.SnapshotVersionInfo.ApplyToArchive(RootArchive);
+			Reader.Execute(RootArchive);
+		}
 	}
 }
 
@@ -130,11 +185,10 @@ int32 UE::LevelSnapshots::Private::FCustomSerializationDataReader::GetNumSubobje
 }
 
 UE::LevelSnapshots::Private::FCustomSerializationDataWriter::FCustomSerializationDataWriter(FCustomSerializationDataGetter_ReadWrite SerializationDataGetter, FWorldSnapshotData& WorldData, UObject* SerializedObject)
-	:
-	UE::LevelSnapshots::Private::FCustomSerializationDataReader(FCustomSerializationDataGetter_ReadOnly(), WorldData),
-	SerializationDataGetter_ReadWrite(SerializationDataGetter),
-	WorldData_ReadWrite(WorldData),
-	SerializedObject(SerializedObject)
+	: FCustomSerializationDataReader(FCustomSerializationDataGetter_ReadOnly(), WorldData)
+	, SerializationDataGetter_ReadWrite(SerializationDataGetter)
+	, WorldData_ReadWrite(WorldData)
+	, SerializedObject(SerializedObject)
 {
 	SerializationDataGetter_ReadOnly = FCustomSerializationDataGetter_ReadOnly::CreateLambda([this]() -> const FCustomSerializationData* { return SerializationDataGetter_ReadWrite.Execute(); });
 	
@@ -143,12 +197,12 @@ UE::LevelSnapshots::Private::FCustomSerializationDataWriter::FCustomSerializatio
 	{
 		const FCustomSubbjectSerializationData& SubobjectData = SerializationData->Subobjects[i];
 		CachedSubobjectMetaData.Add(
-			MakeShared<Internal::FSnapshotSubobjectMetaDataManager>(
+			MakeShared<Internal::FSnapshotSubobjectMetaDataWriter>(
 				WorldData.SerializedObjectReferences[SubobjectData.ObjectPathIndex],
 					i,
-					WorldData_ReadOnly,
+					WorldData,
 					SerializationDataGetter_ReadOnly,
-					FCustomSerializationDataGetter_ReadWrite()
+					SerializationDataGetter_ReadWrite
 				)
 			);
 	}
@@ -158,10 +212,20 @@ void UE::LevelSnapshots::Private::FCustomSerializationDataWriter::WriteObjectAnn
 {
 	if (Writer.IsBound())
 	{
-		FMemoryWriter MemoryWriter(SerializationDataGetter_ReadWrite.Execute()->RootAnnotationData, true);
-		FObjectAndNameAsStringProxyArchive RootArchive(MemoryWriter, false);
-		WorldData_ReadOnly.SnapshotVersionInfo.ApplyToArchive(RootArchive);
-		Writer.Execute(RootArchive);
+		if (WorldData_ReadOnly.SnapshotVersionInfo.GetSnapshotCustomVersion() >= FSnapshotCustomVersion::CustomSubobjectSoftObjectPathRefactor)
+		{
+			FMemoryWriter MemoryWriter(SerializationDataGetter_ReadWrite.Execute()->RootAnnotationData, true);
+			CustomSerialization::FSaveCustomObjectProxyArchive RootArchive(MemoryWriter, WorldData_ReadWrite.SerializedObjectReferences);
+			WorldData_ReadOnly.SnapshotVersionInfo.ApplyToArchive(RootArchive);
+			Writer.Execute(RootArchive);
+		}
+		else
+		{
+			FMemoryWriter MemoryWriter(SerializationDataGetter_ReadWrite.Execute()->RootAnnotationData, true);
+			FObjectAndNameAsStringProxyArchive RootArchive(MemoryWriter, false);
+			WorldData_ReadOnly.SnapshotVersionInfo.ApplyToArchive(RootArchive);
+			Writer.Execute(RootArchive);
+		}
 	}
 }
 
@@ -202,7 +266,7 @@ int32 UE::LevelSnapshots::Private::FCustomSerializationDataWriter::AddSubobjectS
 	// Optimisation: Serialize into the allocated SerializationData directly avoiding a possibly large copy of the serialized data
 	FTakeWorldObjectSnapshotArchive::TakeSnapshot(SerializationData->Subobjects[SubobjectIndex], WorldData_ReadWrite, Subobject);
 	
-	CachedSubobjectMetaData.Add(MakeShared<Internal::FSnapshotSubobjectMetaDataManager>(Subobject, SubobjectIndex, WorldData_ReadWrite, SerializationDataGetter_ReadOnly, SerializationDataGetter_ReadWrite));
+	CachedSubobjectMetaData.Add(MakeShared<Internal::FSnapshotSubobjectMetaDataWriter>(Subobject, SubobjectIndex, WorldData_ReadWrite, SerializationDataGetter_ReadOnly, SerializationDataGetter_ReadWrite));
 
 	TakeSnapshotForSubobject(Subobject, WorldData_ReadWrite);
 	return SubobjectIndex;
diff --git a/Engine/Plugins/VirtualProduction/LevelSnapshots/Source/LevelSnapshots/Private/Data/CustomSerialization/CustomSerializationDataManager.h b/Engine/Plugins/VirtualProduction/LevelSnapshots/Source/LevelSnapshots/Private/Data/CustomSerialization/CustomSerializationDataManager.h
--- a/Engine/Plugins/VirtualProduction/LevelSnapshots/Source/LevelSnapshots/Private/Data/CustomSerialization/CustomSerializationDataManager.h
+++ b/Engine/Plugins/VirtualProduction/LevelSnapshots/Source/LevelSnapshots/Private/Data/CustomSerialization/CustomSerializationDataManager.h
@@ -11,11 +11,6 @@
 
 namespace UE::LevelSnapshots::Private
 {
-	namespace Internal
-	{
-		class FSnapshotSubobjectMetaDataManager;
-	}
-	
 	DECLARE_DELEGATE_RetVal(const FCustomSerializationData*, FCustomSerializationDataGetter_ReadOnly);
 	DECLARE_DELEGATE_RetVal(FCustomSerializationData*, FCustomSerializationDataGetter_ReadWrite);
 
@@ -35,7 +30,7 @@ namespace UE::LevelSnapshots::Private
 	protected:
 		
 		/** To avoid dynamically allocating memory every time GetSubobjectMetaData is called, we cache the results here. Just to improve heap memory footprint a tiny bit. */
-		TArray<TSharedPtr<Internal::FSnapshotSubobjectMetaDataManager>> CachedSubobjectMetaData;
+		TArray<TSharedPtr<ISnapshotSubobjectMetaData>> CachedSubobjectMetaData;
 		
 		/** Looks up actor or subobject  serialization data. Used instead of reference because some functions, like TakeSubobjectSnapshot, cause FWorldSnapshotData::CustomSubobjectSerializationData to reallocate. */
 		FCustomSerializationDataGetter_ReadOnly SerializationDataGetter_ReadOnly;
diff --git a/Engine/Plugins/VirtualProduction/LevelSnapshots/Source/LevelSnapshots/Private/Data/LevelSnapshot.cpp b/Engine/Plugins/VirtualProduction/LevelSnapshots/Source/LevelSnapshots/Private/Data/LevelSnapshot.cpp
--- a/Engine/Plugins/VirtualProduction/LevelSnapshots/Source/LevelSnapshots/Private/Data/LevelSnapshot.cpp
+++ b/Engine/Plugins/VirtualProduction/LevelSnapshots/Source/LevelSnapshots/Private/Data/LevelSnapshot.cpp
@@ -2,8 +2,6 @@
 
 #include "Data/LevelSnapshot.h"
 
-#include "CustomSerialization/CustomSerializationDataManager.h"
-#include "Data/SnapshotCustomVersion.h"
 #include "Data/Util/ActorHashUtil.h"
 #include "Data/Util/EquivalenceUtil.h"
 #include "Data/Util/SnapshotUtil.h"
@@ -13,6 +11,7 @@
 #include "LevelSnapshotsLog.h"
 #include "LevelSnapshotsModule.h"
 #include "Restorability/SnapshotRestorability.h"
+#include "SnapshotCustomVersion.h"
 #include "SnapshotConsoleVariables.h"
 #include "Util/SortedScopedLog.h"
 
@@ -126,7 +125,7 @@ bool ULevelSnapshot::HasChangedSinceSnapshotWasTaken(AActor* WorldActor)
 	const bool bHasChanged = [&]()
 	{
 		// Do not slow down old snapshots by computing hash if they had none saved
-		const bool bHasHashInfo = SerializedData.SnapshotVersionInfo.GetSnapshotCustomVersion() >= UE::LevelSnapshots::Private::FSnapshotCustomVersion::ActorHash;
+		const bool bHasHashInfo = SerializedData.SnapshotVersionInfo.GetSnapshotCustomVersion() >= UE::LevelSnapshots::FSnapshotCustomVersion::ActorHash;
 		// If the object is already loaded, there is no point in wasting time and computing the hash
 		const bool bNeedsHash = Cache.ActorCache.Find(WorldActor) == nullptr;
 		
diff --git a/Engine/Plugins/VirtualProduction/LevelSnapshots/Source/LevelSnapshots/Private/Data/SnapshotCustomVersion.cpp b/Engine/Plugins/VirtualProduction/LevelSnapshots/Source/LevelSnapshots/Private/Data/SnapshotCustomVersion.cpp
--- a/Engine/Plugins/VirtualProduction/LevelSnapshots/Source/LevelSnapshots/Private/Data/SnapshotCustomVersion.cpp
+++ b/Engine/Plugins/VirtualProduction/LevelSnapshots/Source/LevelSnapshots/Private/Data/SnapshotCustomVersion.cpp
@@ -3,5 +3,5 @@
 #include "SnapshotCustomVersion.h"
 #include "Serialization/CustomVersion.h"
 
-const FGuid UE::LevelSnapshots::Private::FSnapshotCustomVersion::GUID(0x3EF5FDBD, 0x21AAAF1F, 0xBCD6F38F, 0xB3A32521);
-FCustomVersionRegistration GRegisterSnapshotCustomVersion(UE::LevelSnapshots::Private::FSnapshotCustomVersion::GUID, UE::LevelSnapshots::Private::FSnapshotCustomVersion::LatestVersion, TEXT("LevelSnapshotVersion"));
+const FGuid UE::LevelSnapshots::FSnapshotCustomVersion::GUID(0x3EF5FDBD, 0x21AAAF1F, 0xBCD6F38F, 0xB3A32521);
+FCustomVersionRegistration GRegisterSnapshotCustomVersion(UE::LevelSnapshots::FSnapshotCustomVersion::GUID, UE::LevelSnapshots::FSnapshotCustomVersion::LatestVersion, TEXT("LevelSnapshotVersion"));
diff --git a/Engine/Plugins/VirtualProduction/LevelSnapshots/Source/LevelSnapshots/Private/Data/SnapshotCustomVersion.h b/Engine/Plugins/VirtualProduction/LevelSnapshots/Source/LevelSnapshots/Private/Data/SnapshotCustomVersion.h
deleted file mode 100644
diff --git a/Engine/Plugins/VirtualProduction/LevelSnapshots/Source/LevelSnapshots/Private/Data/SnapshotVersion.cpp b/Engine/Plugins/VirtualProduction/LevelSnapshots/Source/LevelSnapshots/Private/Data/SnapshotVersion.cpp
--- a/Engine/Plugins/VirtualProduction/LevelSnapshots/Source/LevelSnapshots/Private/Data/SnapshotVersion.cpp
+++ b/Engine/Plugins/VirtualProduction/LevelSnapshots/Source/LevelSnapshots/Private/Data/SnapshotVersion.cpp
@@ -49,7 +49,7 @@ void FSnapshotVersionInfo::Initialize(bool bWithoutSnapshotVersion)
 	for (const FCustomVersion& EngineCustomVersion : AllCurrentVersions.GetAllVersions())
 	{
 		FSnapshotCustomVersionInfo& CustomVersion = CustomVersions.AddDefaulted_GetRef();
-		if (!bWithoutSnapshotVersion || (bWithoutSnapshotVersion && CustomVersion.Key != UE::LevelSnapshots::Private::FSnapshotCustomVersion::GUID))
+		if (!bWithoutSnapshotVersion || (bWithoutSnapshotVersion && CustomVersion.Key != UE::LevelSnapshots::FSnapshotCustomVersion::GUID))
 		{
 			CustomVersion.Initialize(EngineCustomVersion);
 		}
@@ -86,7 +86,7 @@ int32 FSnapshotVersionInfo::GetSnapshotCustomVersion() const
 {
 	for (const FSnapshotCustomVersionInfo& CustomVersion : CustomVersions)
 	{
-		if (CustomVersion.Key == UE::LevelSnapshots::Private::FSnapshotCustomVersion::GUID)
+		if (CustomVersion.Key == UE::LevelSnapshots::FSnapshotCustomVersion::GUID)
 		{
 			return CustomVersion.Version;
 		}
diff --git a/Engine/Plugins/VirtualProduction/LevelSnapshots/Source/LevelSnapshots/Private/Data/WorldSnapshotData.cpp b/Engine/Plugins/VirtualProduction/LevelSnapshots/Source/LevelSnapshots/Private/Data/WorldSnapshotData.cpp
--- a/Engine/Plugins/VirtualProduction/LevelSnapshots/Source/LevelSnapshots/Private/Data/WorldSnapshotData.cpp
+++ b/Engine/Plugins/VirtualProduction/LevelSnapshots/Source/LevelSnapshots/Private/Data/WorldSnapshotData.cpp
@@ -192,7 +192,7 @@ namespace UE::LevelSnapshots::Private::Internal
 bool FWorldSnapshotData::Serialize(FArchive& Ar)
 {
 	using namespace UE::LevelSnapshots::Private;
-	Ar.UsingCustomVersion(FSnapshotCustomVersion::GUID);
+	Ar.UsingCustomVersion(UE::LevelSnapshots::FSnapshotCustomVersion::GUID);
 
 	// When this struct is saved, the save algorithm collects references. It's faster if we just give it the info directly.
 	if (Ar.IsObjectReferenceCollector())
@@ -213,7 +213,7 @@ bool FWorldSnapshotData::Serialize(FArchive& Ar)
 
 	if (Ar.IsLoading())
 	{
-		if (Ar.CustomVer(FSnapshotCustomVersion::GUID) >= FSnapshotCustomVersion::OoddleCompression)
+		if (Ar.CustomVer(UE::LevelSnapshots::FSnapshotCustomVersion::GUID) >= UE::LevelSnapshots::FSnapshotCustomVersion::OoddleCompression)
 		{
 			UE::LevelSnapshots::Decompress(Ar, *this);
 			return true;
@@ -242,7 +242,7 @@ void FWorldSnapshotData::PostSerialize(const FArchive& Ar)
 	}
 	
 	if (Ar.IsLoading()
-		&& Ar.CustomVer(FSnapshotCustomVersion::GUID) < FSnapshotCustomVersion::ClassArchetypeRefactor)
+		&& Ar.CustomVer(UE::LevelSnapshots::FSnapshotCustomVersion::GUID) < UE::LevelSnapshots::FSnapshotCustomVersion::ClassArchetypeRefactor)
 	{
 		MigrateClassDefaultData(*this);
 	}
diff --git a/Engine/Plugins/VirtualProduction/LevelSnapshots/Source/LevelSnapshots/Private/Restorability/ActorGroupRestoration.cpp b/Engine/Plugins/VirtualProduction/LevelSnapshots/Source/LevelSnapshots/Private/Restorability/ActorGroupRestoration.cpp
--- a/Engine/Plugins/VirtualProduction/LevelSnapshots/Source/LevelSnapshots/Private/Restorability/ActorGroupRestoration.cpp
+++ b/Engine/Plugins/VirtualProduction/LevelSnapshots/Source/LevelSnapshots/Private/Restorability/ActorGroupRestoration.cpp
@@ -220,7 +220,7 @@ namespace UE::LevelSnapshots::Private::ActorGroupRestoration
 
 				// Hacky way of forcing GroupActor to reconsider whether it should exist...
 				// References can fail to restore ... possibly there are no children
-				GroupActor->Remove(*GroupActor);
+				GroupActor->Remove(*Cast<AGroupActor>(Params.Actor)); // Silly cast used here to fix silly static analysis warning...
 			}
 		}
 
diff --git a/Engine/Plugins/VirtualProduction/LevelSnapshots/Source/LevelSnapshots/Private/Restorability/SnapshotRestorability.cpp b/Engine/Plugins/VirtualProduction/LevelSnapshots/Source/LevelSnapshots/Private/Restorability/SnapshotRestorability.cpp
--- a/Engine/Plugins/VirtualProduction/LevelSnapshots/Source/LevelSnapshots/Private/Restorability/SnapshotRestorability.cpp
+++ b/Engine/Plugins/VirtualProduction/LevelSnapshots/Source/LevelSnapshots/Private/Restorability/SnapshotRestorability.cpp
@@ -226,7 +226,7 @@ bool UE::LevelSnapshots::Restorability::ShouldConsiderRemovedActorForRecreation(
 	const bool bTransacts = (Params.ObjectFlags& RF_Transactional) != 0;
 	const bool bWereFlagsSaved = Params.ObjectFlags != 0 // Objects usually have at least one flag so if it's 0 then it is very likely the flags were not saved...
 		// ... support for saving flags was added some time before ClassArchetypeRefactor. We check against ClassArchetypeRefactor for the unlikely case that in the future an object really has no flags.
-		|| Params.WorldData.SnapshotVersionInfo.GetSnapshotCustomVersion() >= LevelSnapshots::Private::FSnapshotCustomVersion::ClassArchetypeRefactor;
+		|| Params.WorldData.SnapshotVersionInfo.GetSnapshotCustomVersion() >= FSnapshotCustomVersion::ClassArchetypeRefactor;
 	
 	return (bTransacts || !bWereFlagsSaved)
 		&& LevelSnapshots::Private::FLevelSnapshotsModule::GetInternalModuleInstance().CanRecreateActor(Params);
diff --git a/Engine/Plugins/VirtualProduction/LevelSnapshots/Source/LevelSnapshots/Private/Tests/CustomObjectSerializationTests.cpp b/Engine/Plugins/VirtualProduction/LevelSnapshots/Source/LevelSnapshots/Private/Tests/CustomObjectSerializationTests.cpp
--- a/Engine/Plugins/VirtualProduction/LevelSnapshots/Source/LevelSnapshots/Private/Tests/CustomObjectSerializationTests.cpp
+++ b/Engine/Plugins/VirtualProduction/LevelSnapshots/Source/LevelSnapshots/Private/Tests/CustomObjectSerializationTests.cpp
@@ -117,11 +117,19 @@ namespace UE::LevelSnapshots::Private::Tests
 				CallOrder.Add(EFunctionCall::OnPostSerializeEditorSubobject);
 			}
 		
-			virtual void PreApplySnapshotProperties(UObject* OriginalObject, const ICustomSnapshotSerializationData& DataStorage) override
+			virtual void PreApplyToSnapshotObject(UObject* OriginalObject, const ICustomSnapshotSerializationData& DataStorage) override
 			{
 				CallOrder.Add(EFunctionCall::PreApplySnapshotProperties);
 			}
-			virtual void PostApplySnapshotProperties(UObject* OriginalObject, const ICustomSnapshotSerializationData& DataStorage) override
+			virtual void PostApplyToSnapshotObject(UObject* OriginalObject, const ICustomSnapshotSerializationData& DataStorage) override
+			{
+				CallOrder.Add(EFunctionCall::PostApplySnapshotProperties);
+			}
+			virtual void PreApplyToEditorObject(UObject* OriginalObject, const ICustomSnapshotSerializationData& DataStorage, const FPropertySelectionMap& SelectionMap) override
+			{
+				CallOrder.Add(EFunctionCall::PreApplySnapshotProperties);
+			}
+			virtual void PostApplyToEditorObject(UObject* OriginalObject, const ICustomSnapshotSerializationData& DataStorage, const FPropertySelectionMap& SelectionMap) override
 			{
 				CallOrder.Add(EFunctionCall::PostApplySnapshotProperties);
 			}
@@ -244,7 +252,7 @@ namespace UE::LevelSnapshots::Private::Tests
 				return Actor->NonReflectedSubobject.Get();
 			}
 
-			virtual void PostApplySnapshotProperties(UObject* OriginalObject, const ICustomSnapshotSerializationData& DataStorage) override
+			virtual void PostApplyToSnapshotObject(UObject* OriginalObject, const ICustomSnapshotSerializationData& DataStorage) override
 			{
 				ASnapshotTestActor* Actor = Cast<ASnapshotTestActor>(OriginalObject);
 				DataStorage.ReadObjectAnnotation(FObjectAnnotator::CreateLambda([Actor](FArchive& Archive)
diff --git a/Engine/Plugins/VirtualProduction/LevelSnapshots/Source/LevelSnapshots/Public/Data/SnapshotCustomVersion.h b/Engine/Plugins/VirtualProduction/LevelSnapshots/Source/LevelSnapshots/Public/Data/SnapshotCustomVersion.h
new file mode 100644
--- /dev/null
+++ b/Engine/Plugins/VirtualProduction/LevelSnapshots/Source/LevelSnapshots/Public/Data/SnapshotCustomVersion.h
@@ -0,0 +1,42 @@
+// Copyright Epic Games, Inc. All Rights Reserved.
+
+#pragma once
+
+#include "CoreMinimal.h"
+#include "Misc/Guid.h"
+
+namespace UE::LevelSnapshots
+{
+	struct LEVELSNAPSHOTS_API FSnapshotCustomVersion
+	{
+		enum Type
+		{
+			/** Before any version changes were made in the plugin */
+			BeforeCustomVersionWasAdded = 0,
+
+			/** When subobject support was added. Specifically, USceneComponent::AttachParent were not captured. */
+			SubobjectSupport = 1,
+
+			/** FSnapshotActorData now stores actor hash data to facilitate checking whether an actor has changed without loading the actor */
+			ActorHash = 2,
+
+			/** FWorldSnapshotData::ClassDefaults was replaced by FWorldSnapshotData::ClassData. */
+			ClassArchetypeRefactor = 3,
+
+			/** FWorldSnapshotData now compresses data using oodle before it is saved to disk */
+			OoddleCompression = 4,
+
+			/** Before this soft object references were saved by FCustomSerializationDataWriter using FObjectAndNameAsStringProxyArchive - if you moved the asset our data would get corrupted. */
+			CustomSubobjectSoftObjectPathRefactor = 5,
+
+			// -----<new versions can be added above this line>-------------------------------------------------
+			VersionPlusOne,
+			LatestVersion = VersionPlusOne - 1
+		};
+
+		const static FGuid GUID;
+
+		private:
+		FSnapshotCustomVersion() = delete;
+	};
+}
diff --git a/Engine/Plugins/VirtualProduction/LevelSnapshots/Source/LevelSnapshots/Public/Restorability/Interfaces/ICustomObjectSnapshotSerializer.h b/Engine/Plugins/VirtualProduction/LevelSnapshots/Source/LevelSnapshots/Public/Restorability/Interfaces/ICustomObjectSnapshotSerializer.h
--- a/Engine/Plugins/VirtualProduction/LevelSnapshots/Source/LevelSnapshots/Public/Restorability/Interfaces/ICustomObjectSnapshotSerializer.h
+++ b/Engine/Plugins/VirtualProduction/LevelSnapshots/Source/LevelSnapshots/Public/Restorability/Interfaces/ICustomObjectSnapshotSerializer.h
@@ -91,14 +91,29 @@ namespace UE::LevelSnapshots
 
 
 		
-		/** Optional. Called before properties are applied to an object for which this interface was registered. Called on both snapshot and editor world versions. */
-		virtual void PreApplySnapshotProperties(UObject* Object, const ICustomSnapshotSerializationData& DataStorage) {}
+		/** Optional. Called before properties are applied to anobject for which this interface was registered. Called on the object in the transient snapshot world. */
+		virtual void PreApplyToSnapshotObject(UObject* Object, const ICustomSnapshotSerializationData& DataStorage) {}
 		
-		/** Optional. Called after properties are applied to an object for which this interface was registered. Called on both snapshot and editor world versions. */
-		virtual void PostApplySnapshotProperties(UObject* Object, const ICustomSnapshotSerializationData& DataStorage) {}
+		/** Optional. Called after properties are applied to an object for which this interface was registered. Called on the object in the transient snapshot world. */
+		virtual void PostApplyToSnapshotObject(UObject* Object, const ICustomSnapshotSerializationData& DataStorage) {}
 
 
 		
+		/** Optional. Called before properties are applied to anobject for which this interface was registered. Called on the object in the editor world. */
+		virtual void PreApplyToEditorObject(UObject* Object, const ICustomSnapshotSerializationData& DataStorage, const FPropertySelectionMap& SelectionMap)
+		{
+			// Calling it for convenience... usually the implementation is equal
+			PreApplyToSnapshotObject(Object, DataStorage);
+		}
+		
+		/** Optional. Called after properties are applied to an object for which this interface was registered. Called on the object in the editor world. */
+		virtual void PostApplyToEditorObject(UObject* Object, const ICustomSnapshotSerializationData& DataStorage, const FPropertySelectionMap& SelectionMap)
+		{
+			// Calling it for convenience... usually the implementation is equal
+			PostApplyToSnapshotObject(Object, DataStorage);
+		}
+
+		
 		virtual ~ICustomObjectSnapshotSerializer() = default;
 	};
 }
