ObjectPtr upgrade for engine plugins used from UnrealEditor -allmodules.  Automated upgrade process.

diff --git a/Engine/Plugins/VirtualProduction/LevelSnapshots/Source/LevelSnapshots/Private/Tests/Types/SnapshotTestActor.h b/Engine/Plugins/VirtualProduction/LevelSnapshots/Source/LevelSnapshots/Private/Tests/Types/SnapshotTestActor.h
--- a/Engine/Plugins/VirtualProduction/LevelSnapshots/Source/LevelSnapshots/Private/Tests/Types/SnapshotTestActor.h
+++ b/Engine/Plugins/VirtualProduction/LevelSnapshots/Source/LevelSnapshots/Private/Tests/Types/SnapshotTestActor.h
@@ -40,10 +40,10 @@ public:
 	float FloatProperty;
 
 	UPROPERTY(EditAnywhere, Instanced, Category = "Level Snapshots")
-	USubSubobject* NestedChild;
+	TObjectPtr<USubSubobject> NestedChild;
 
 	UPROPERTY(Instanced)
-	USubSubobject* UneditableNestedChild;
+	TObjectPtr<USubSubobject> UneditableNestedChild;
 
 	FName NonReflectedName;
 	UObject* NonReflectedObjectProperty;
@@ -65,7 +65,7 @@ public:
 	float FloatProperty;
 
 	UPROPERTY(Instanced)
-	USubobject* Subobject;
+	TObjectPtr<USubobject> Subobject;
 };
 
 UCLASS()
@@ -111,16 +111,16 @@ public:
 	/******************** Raw references  ********************/
 	
 	UPROPERTY(EditAnywhere, Category = "Level Snapshots")
-	UObject* ObjectReference;
+	TObjectPtr<UObject> ObjectReference;
 
 	UPROPERTY(EditAnywhere, Category = "Level Snapshots")
-	TArray<UObject*> ObjectArray;
+	TArray<TObjectPtr<UObject>> ObjectArray;
 
 	UPROPERTY(EditAnywhere, Category = "Level Snapshots")
-	TSet<UObject*> ObjectSet;
+	TSet<TObjectPtr<UObject>> ObjectSet;
 
 	UPROPERTY(EditAnywhere, Category = "Level Snapshots")
-	TMap<FName, UObject*> ObjectMap;
+	TMap<FName, TObjectPtr<UObject>> ObjectMap;
 
 
 
@@ -177,10 +177,10 @@ public:
 	/******************** External component references  ********************/
 	
 	UPROPERTY(EditAnywhere, Category = "Level Snapshots")
-	UActorComponent* ExternalComponentReference;
+	TObjectPtr<UActorComponent> ExternalComponentReference;
 
 	UPROPERTY(EditAnywhere, Category = "Level Snapshots")
-	UObject* ExternalComponentReferenceAsUObject;
+	TObjectPtr<UObject> ExternalComponentReferenceAsUObject;
 
 
 
@@ -188,16 +188,16 @@ public:
 	/******************** External references  ********************/
 	
 	UPROPERTY(EditAnywhere, Category = "Level Snapshots")
-	UMaterialInterface* GradientLinearMaterial;
+	TObjectPtr<UMaterialInterface> GradientLinearMaterial;
 
 	UPROPERTY(EditAnywhere, Category = "Level Snapshots")
-	UMaterialInterface* GradientRadialMaterial;
+	TObjectPtr<UMaterialInterface> GradientRadialMaterial;
 	
 	UPROPERTY(EditAnywhere, Category = "Level Snapshots")
-	UStaticMesh* CubeMesh;
+	TObjectPtr<UStaticMesh> CubeMesh;
 
 	UPROPERTY(EditAnywhere, Category = "Level Snapshots")
-	UStaticMesh* CylinderMesh;
+	TObjectPtr<UStaticMesh> CylinderMesh;
 
 
 
@@ -205,44 +205,44 @@ public:
 	/******************** Subobject Component references  ********************/
 
 	UPROPERTY(EditAnywhere, Category = "Level Snapshots")
-	UStaticMeshComponent* StaticMeshComponent;
+	TObjectPtr<UStaticMeshComponent> StaticMeshComponent;
 	
 	UPROPERTY(EditAnywhere, Category = "Level Snapshots")
-	UInstancedStaticMeshComponent* InstancedMeshComponent;
+	TObjectPtr<UInstancedStaticMeshComponent> InstancedMeshComponent;
 
 	UPROPERTY(EditAnywhere, Category = "Level Snapshots")
-	UPointLightComponent* PointLightComponent;
+	TObjectPtr<UPointLightComponent> PointLightComponent;
 
 	UPROPERTY(EditAnywhere, Category = "Level Snapshots")
-	USnapshotTestComponent* TestComponent;
+	TObjectPtr<USnapshotTestComponent> TestComponent;
 
 	
 	/******************** Subobject references  ********************/
 
 	UPROPERTY(EditAnywhere, Instanced, Category = "Level Snapshots")
-	USubobject* EditableInstancedSubobject_DefaultSubobject;
+	TObjectPtr<USubobject> EditableInstancedSubobject_DefaultSubobject;
 
 	UPROPERTY(Instanced)
-	USubobject* InstancedOnlySubobject_DefaultSubobject;
+	TObjectPtr<USubobject> InstancedOnlySubobject_DefaultSubobject;
 
 	UPROPERTY()
-	USubobject* NakedSubobject_DefaultSubobject;
+	TObjectPtr<USubobject> NakedSubobject_DefaultSubobject;
 
 	UPROPERTY(EditAnywhere, Category = "Level Snapshots")
-	USubobject* EditOnlySubobject_OptionalSubobject;
+	TObjectPtr<USubobject> EditOnlySubobject_OptionalSubobject;
 	
 	
 	UPROPERTY(EditAnywhere, Instanced, Category = "Level Snapshots")
-	TArray<USubobject*> EditableInstancedSubobjectArray_OptionalSubobject;
+	TArray<TObjectPtr<USubobject>> EditableInstancedSubobjectArray_OptionalSubobject;
 
 	UPROPERTY(EditAnywhere, Instanced, Category = "Level Snapshots")
-	TMap<FName, USubobject*> EditableInstancedSubobjectMap_OptionalSubobject;
+	TMap<FName, TObjectPtr<USubobject>> EditableInstancedSubobjectMap_OptionalSubobject;
 
 	UPROPERTY(EditAnywhere, Category = "Level Snapshots")
-	TArray<USubobject*> EditOnlySubobjectArray_OptionalSubobject;
+	TArray<TObjectPtr<USubobject>> EditOnlySubobjectArray_OptionalSubobject;
 
 	UPROPERTY(EditAnywhere, Category = "Level Snapshots")
-	TMap<FName, USubobject*> EditOnlySubobjectMap_OptionalSubobject;
+	TMap<FName, TObjectPtr<USubobject>> EditOnlySubobjectMap_OptionalSubobject;
 	
 	TStrongObjectPtr<USubobject> NonReflectedSubobject;
 
diff --git a/Engine/Plugins/VirtualProduction/LevelSnapshots/Source/LevelSnapshots/Public/Data/LevelSnapshot.h b/Engine/Plugins/VirtualProduction/LevelSnapshots/Source/LevelSnapshots/Public/Data/LevelSnapshot.h
--- a/Engine/Plugins/VirtualProduction/LevelSnapshots/Source/LevelSnapshots/Public/Data/LevelSnapshot.h
+++ b/Engine/Plugins/VirtualProduction/LevelSnapshots/Source/LevelSnapshots/Public/Data/LevelSnapshot.h
@@ -138,10 +138,10 @@ private:
 #endif
 	/** The root world, equivalent to the persistent world, we will be adding temporary actors to */
 	UPROPERTY(Transient)
-	UWorld* RootSnapshotWorld;
+	TObjectPtr<UWorld> RootSnapshotWorld;
 	/** Sublevels of RootSnapshotWorld */
 	UPROPERTY(Transient)
-	TArray<UWorld*> SnapshotSublevels;
+	TArray<TObjectPtr<UWorld>> SnapshotSublevels;
 	
 	/** The saved snapshot data */
 	UPROPERTY()
diff --git a/Engine/Plugins/VirtualProduction/LevelSnapshots/Source/LevelSnapshots/Public/Data/WorldSnapshotData.h b/Engine/Plugins/VirtualProduction/LevelSnapshots/Source/LevelSnapshots/Public/Data/WorldSnapshotData.h
--- a/Engine/Plugins/VirtualProduction/LevelSnapshots/Source/LevelSnapshots/Public/Data/WorldSnapshotData.h
+++ b/Engine/Plugins/VirtualProduction/LevelSnapshots/Source/LevelSnapshots/Public/Data/WorldSnapshotData.h
@@ -38,7 +38,7 @@ struct LEVELSNAPSHOTS_API FWorldSnapshotData
 	TWeakObjectPtr<UWorld> SnapshotWorld;
 	/** Sublevels of SnapshotWorld; excludes SnapshotWorld->PersistentLevel */
 	UPROPERTY(Transient)
-	TArray<UWorld*> SnapshotSublevels;
+	TArray<TObjectPtr<UWorld>> SnapshotSublevels;
 
 	/**
 	 * Stores versioning information we inject into archives.
diff --git a/Engine/Plugins/VirtualProduction/LevelSnapshots/Source/LevelSnapshotsEditor/Private/Data/FilterLoader.h b/Engine/Plugins/VirtualProduction/LevelSnapshots/Source/LevelSnapshotsEditor/Private/Data/FilterLoader.h
--- a/Engine/Plugins/VirtualProduction/LevelSnapshots/Source/LevelSnapshotsEditor/Private/Data/FilterLoader.h
+++ b/Engine/Plugins/VirtualProduction/LevelSnapshots/Source/LevelSnapshotsEditor/Private/Data/FilterLoader.h
@@ -49,5 +49,5 @@ private:
 	FSoftObjectPath AssetLastSavedOrLoaded = nullptr;
 
 	UPROPERTY()
-	ULevelSnapshotsFilterPreset* AssetBeingEdited;
+	TObjectPtr<ULevelSnapshotsFilterPreset> AssetBeingEdited;
 };
diff --git a/Engine/Plugins/VirtualProduction/LevelSnapshots/Source/LevelSnapshotsEditor/Private/Data/Filters/ConjunctionFilter.h b/Engine/Plugins/VirtualProduction/LevelSnapshots/Source/LevelSnapshotsEditor/Private/Data/Filters/ConjunctionFilter.h
--- a/Engine/Plugins/VirtualProduction/LevelSnapshots/Source/LevelSnapshotsEditor/Private/Data/Filters/ConjunctionFilter.h
+++ b/Engine/Plugins/VirtualProduction/LevelSnapshots/Source/LevelSnapshotsEditor/Private/Data/Filters/ConjunctionFilter.h
@@ -50,7 +50,7 @@ public:
 private:
 
 	UPROPERTY(Instanced)
-	TArray<UNegatableFilter*> Children;
+	TArray<TObjectPtr<UNegatableFilter>> Children;
 
 	UPROPERTY()
 	bool bIgnoreFilter = false;
diff --git a/Engine/Plugins/VirtualProduction/LevelSnapshots/Source/LevelSnapshotsEditor/Private/Data/Filters/LevelSnapshotsFilterPreset.h b/Engine/Plugins/VirtualProduction/LevelSnapshots/Source/LevelSnapshotsEditor/Private/Data/Filters/LevelSnapshotsFilterPreset.h
--- a/Engine/Plugins/VirtualProduction/LevelSnapshots/Source/LevelSnapshotsEditor/Private/Data/Filters/LevelSnapshotsFilterPreset.h
+++ b/Engine/Plugins/VirtualProduction/LevelSnapshots/Source/LevelSnapshotsEditor/Private/Data/Filters/LevelSnapshotsFilterPreset.h
@@ -65,7 +65,7 @@ public:
 private:
 	
 	UPROPERTY(Instanced)
-	TArray<UConjunctionFilter*> Children;
+	TArray<TObjectPtr<UConjunctionFilter>> Children;
 
 	bool bHasJustCreatedNewChild = false;
 	
diff --git a/Engine/Plugins/VirtualProduction/LevelSnapshots/Source/LevelSnapshotsEditor/Private/Data/Filters/NegatableFilter.h b/Engine/Plugins/VirtualProduction/LevelSnapshots/Source/LevelSnapshotsEditor/Private/Data/Filters/NegatableFilter.h
--- a/Engine/Plugins/VirtualProduction/LevelSnapshots/Source/LevelSnapshotsEditor/Private/Data/Filters/NegatableFilter.h
+++ b/Engine/Plugins/VirtualProduction/LevelSnapshots/Source/LevelSnapshotsEditor/Private/Data/Filters/NegatableFilter.h
@@ -65,5 +65,5 @@ private:
 	bool bIgnoreFilter = false;
 	
 	UPROPERTY(Instanced)
-	ULevelSnapshotFilter* ChildFilter;
+	TObjectPtr<ULevelSnapshotFilter> ChildFilter;
 };
\ No newline at end of file
diff --git a/Engine/Plugins/VirtualProduction/LevelSnapshots/Source/LevelSnapshotsEditor/Private/Data/LevelSnapshotsEditorData.h b/Engine/Plugins/VirtualProduction/LevelSnapshots/Source/LevelSnapshotsEditor/Private/Data/LevelSnapshotsEditorData.h
--- a/Engine/Plugins/VirtualProduction/LevelSnapshots/Source/LevelSnapshotsEditor/Private/Data/LevelSnapshotsEditorData.h
+++ b/Engine/Plugins/VirtualProduction/LevelSnapshots/Source/LevelSnapshotsEditor/Private/Data/LevelSnapshotsEditorData.h
@@ -90,13 +90,13 @@ private:
 	FDelegateHandle TrackedFilterModifiedHandle;
 	
 	UPROPERTY()
-	UFavoriteFilterContainer* FavoriteFilters;
+	TObjectPtr<UFavoriteFilterContainer> FavoriteFilters;
 	/* Stores user-defined filters in chain of ORs of ANDs. */
 	UPROPERTY()
-	ULevelSnapshotsFilterPreset* UserDefinedFilters;
+	TObjectPtr<ULevelSnapshotsFilterPreset> UserDefinedFilters;
 	/* Handles save & load requests for exchanging UserDefinedFilters. */
 	UPROPERTY()
-	UFilterLoader* FilterLoader;
+	TObjectPtr<UFilterLoader> FilterLoader;
 	
 	/* Used for determining whether the filter state has changed since it was last refreshed. */
 	UPROPERTY()
@@ -104,7 +104,7 @@ private:
 
 	/* Converts UserDefinedFilters into ULevelSnapshotsSelectionSet display in results view. */
 	UPROPERTY()
-	UFilteredResults* FilterResults;
+	TObjectPtr<UFilteredResults> FilterResults;
 
 	/* Snapshot selected by user */
 	TOptional<TStrongObjectPtr<ULevelSnapshot>> ActiveSnapshot;
diff --git a/Engine/Plugins/VirtualProduction/LiveLinkCamera/Source/LiveLinkCamera/Public/LiveLinkCameraController.h b/Engine/Plugins/VirtualProduction/LiveLinkCamera/Source/LiveLinkCamera/Public/LiveLinkCameraController.h
--- a/Engine/Plugins/VirtualProduction/LiveLinkCamera/Source/LiveLinkCamera/Public/LiveLinkCameraController.h
+++ b/Engine/Plugins/VirtualProduction/LiveLinkCamera/Source/LiveLinkCamera/Public/LiveLinkCameraController.h
@@ -150,7 +150,7 @@ protected:
 
 	UE_DEPRECATED(5.1, "This property has been deprecated. The LiveLink camera controller no longer evaluated distortion.")
 	UPROPERTY(Transient)
-	ULensDistortionModelHandlerBase* LensDistortionHandler_DEPRECATED = nullptr;
+	TObjectPtr<ULensDistortionModelHandlerBase> LensDistortionHandler_DEPRECATED = nullptr;
 
 	UE_DEPRECATED(5.1, "This property has been deprecated. The LiveLink camera controller no longer evaluated distortion.")
 	UPROPERTY(DuplicateTransient)
diff --git a/Engine/Plugins/VirtualProduction/LiveLinkCamera/Source/LiveLinkCameraRecording/Private/MovieSceneLiveLinkCameraControllerSection.h b/Engine/Plugins/VirtualProduction/LiveLinkCamera/Source/LiveLinkCameraRecording/Private/MovieSceneLiveLinkCameraControllerSection.h
--- a/Engine/Plugins/VirtualProduction/LiveLinkCamera/Source/LiveLinkCameraRecording/Private/MovieSceneLiveLinkCameraControllerSection.h
+++ b/Engine/Plugins/VirtualProduction/LiveLinkCamera/Source/LiveLinkCameraRecording/Private/MovieSceneLiveLinkCameraControllerSection.h
@@ -39,5 +39,5 @@ public:
 private:
 	/** Saved duplicate of the LensFile asset used by the recorded LiveLink Camera Controller at the time of recording */
 	UPROPERTY(VisibleAnywhere, Category="Camera Calibration")
-	ULensFile* CachedLensFile = nullptr;
+	TObjectPtr<ULensFile> CachedLensFile = nullptr;
 };
diff --git a/Engine/Plugins/VirtualProduction/LiveLinkLens/Source/LiveLinkLens/Public/LiveLinkLensController.h b/Engine/Plugins/VirtualProduction/LiveLinkLens/Source/LiveLinkLens/Public/LiveLinkLensController.h
--- a/Engine/Plugins/VirtualProduction/LiveLinkLens/Source/LiveLinkLens/Public/LiveLinkLensController.h
+++ b/Engine/Plugins/VirtualProduction/LiveLinkLens/Source/LiveLinkLens/Public/LiveLinkLensController.h
@@ -38,7 +38,7 @@ protected:
 #if WITH_EDITORONLY_DATA
 	UE_DEPRECATED(5.1, "This property has been deprecated. Distortion is handled directly by controlled Lens Component.")
  	UPROPERTY(Transient)
- 	ULensDistortionModelHandlerBase* LensDistortionHandler_DEPRECATED = nullptr;
+ 	TObjectPtr<ULensDistortionModelHandlerBase> LensDistortionHandler_DEPRECATED = nullptr;
 
 	UE_DEPRECATED(5.1, "This property has been deprecated and is no longer used.")
 	UPROPERTY(DuplicateTransient)
diff --git a/Engine/Plugins/VirtualProduction/RemoteControl/Source/RemoteControl/Private/RemoteControlExposeRegistry.h b/Engine/Plugins/VirtualProduction/RemoteControl/Source/RemoteControl/Private/RemoteControlExposeRegistry.h
--- a/Engine/Plugins/VirtualProduction/RemoteControl/Source/RemoteControl/Private/RemoteControlExposeRegistry.h
+++ b/Engine/Plugins/VirtualProduction/RemoteControl/Source/RemoteControl/Private/RemoteControlExposeRegistry.h
@@ -218,5 +218,5 @@ private:
 
 	/** Holds the types of entities exposed in the registry. */
 	UPROPERTY()
-	TSet<UScriptStruct*> ExposedTypes;
+	TSet<TObjectPtr<UScriptStruct>> ExposedTypes;
 };
diff --git a/Engine/Plugins/VirtualProduction/RemoteControl/Source/RemoteControl/Public/RemoteControlField.h b/Engine/Plugins/VirtualProduction/RemoteControl/Source/RemoteControl/Public/RemoteControlField.h
--- a/Engine/Plugins/VirtualProduction/RemoteControl/Source/RemoteControl/Public/RemoteControlField.h
+++ b/Engine/Plugins/VirtualProduction/RemoteControl/Source/RemoteControl/Public/RemoteControlField.h
@@ -253,7 +253,7 @@ public:
 	 */
 #if WITH_EDITORONLY_DATA
 	UPROPERTY()
-	mutable UFunction* Function_DEPRECATED = nullptr;
+	mutable TObjectPtr<UFunction> Function_DEPRECATED = nullptr;
 #endif
 
 	/**
diff --git a/Engine/Plugins/VirtualProduction/RemoteControl/Source/RemoteControl/Public/RemoteControlInstanceMaterial.h b/Engine/Plugins/VirtualProduction/RemoteControl/Source/RemoteControl/Public/RemoteControlInstanceMaterial.h
--- a/Engine/Plugins/VirtualProduction/RemoteControl/Source/RemoteControl/Public/RemoteControlInstanceMaterial.h
+++ b/Engine/Plugins/VirtualProduction/RemoteControl/Source/RemoteControl/Public/RemoteControlInstanceMaterial.h
@@ -45,7 +45,7 @@ public:
 public:
 	/** Stores original binding class */
 	UPROPERTY()
-	UClass* OriginalClass = nullptr;
+	TObjectPtr<UClass> OriginalClass = nullptr;
 
 	/** Stores original material parameter info */
 	UPROPERTY()
diff --git a/Engine/Plugins/VirtualProduction/RemoteControl/Source/RemoteControl/Public/RemoteControlPreset.h b/Engine/Plugins/VirtualProduction/RemoteControl/Source/RemoteControl/Public/RemoteControlPreset.h
--- a/Engine/Plugins/VirtualProduction/RemoteControl/Source/RemoteControl/Public/RemoteControlPreset.h
+++ b/Engine/Plugins/VirtualProduction/RemoteControl/Source/RemoteControl/Public/RemoteControlPreset.h
@@ -611,7 +611,7 @@ public:
 
 	/** This preset's list of objects that are exposed or that have exposed fields. */
 	UPROPERTY(EditAnywhere, Category = "Remote Control Preset")
-	TArray<URemoteControlBinding*> Bindings;
+	TArray<TObjectPtr<URemoteControlBinding>> Bindings;
 
 	/** ~~~Virtual Property Wrapper Functions ~~~
 	* 
@@ -768,7 +768,7 @@ private:
 
 	UPROPERTY(Instanced)
 	/** Holds exposed entities on the preset. */
-	URemoteControlExposeRegistry* Registry = nullptr;
+	TObjectPtr<URemoteControlExposeRegistry> Registry = nullptr;
 
 	/** Delegate triggered when an entity is exposed. */
 	FOnPresetEntityEvent OnEntityExposedDelegate;
diff --git a/Engine/Plugins/VirtualProduction/RemoteControl/Source/RemoteControl/Public/RemoteControlPresetActor.h b/Engine/Plugins/VirtualProduction/RemoteControl/Source/RemoteControl/Public/RemoteControlPresetActor.h
--- a/Engine/Plugins/VirtualProduction/RemoteControl/Source/RemoteControl/Public/RemoteControlPresetActor.h
+++ b/Engine/Plugins/VirtualProduction/RemoteControl/Source/RemoteControl/Public/RemoteControlPresetActor.h
@@ -18,7 +18,7 @@ public:
 
 public:
 	UPROPERTY(EditAnywhere, Category="General")
-	URemoteControlPreset* Preset;
+	TObjectPtr<URemoteControlPreset> Preset;
 
 private:
 	/** Delegate handle for actor deletion callback. */
diff --git a/Engine/Plugins/VirtualProduction/RemoteControl/Source/RemoteControlLogic/Public/Behaviour/Builtin/Conditional/RCBehaviourConditional.h b/Engine/Plugins/VirtualProduction/RemoteControl/Source/RemoteControlLogic/Public/Behaviour/Builtin/Conditional/RCBehaviourConditional.h
--- a/Engine/Plugins/VirtualProduction/RemoteControl/Source/RemoteControlLogic/Public/Behaviour/Builtin/Conditional/RCBehaviourConditional.h
+++ b/Engine/Plugins/VirtualProduction/RemoteControl/Source/RemoteControlLogic/Public/Behaviour/Builtin/Conditional/RCBehaviourConditional.h
@@ -90,7 +90,7 @@ public:
 	/** Data storage for Actions and related Conditions; stored as a mapping of Action object and associated condition data
 	*  Each Action is associated with a unique condition (for the Conditional Behaviour) */
 	UPROPERTY()
-	TMap<URCAction*, FRCBehaviourCondition> Conditions;
+	TMap<TObjectPtr<URCAction>, FRCBehaviourCondition> Conditions;
 
 	/** Virtual property used to build the Comparand - i.e. the property with which the Controller will be compared for a given condition*/
 	UPROPERTY()
diff --git a/Engine/Plugins/VirtualProduction/RemoteControl/Source/RemoteControlLogic/Public/Behaviour/Builtin/RangeMap/RCRangeMapBehaviour.h b/Engine/Plugins/VirtualProduction/RemoteControl/Source/RemoteControlLogic/Public/Behaviour/Builtin/RangeMap/RCRangeMapBehaviour.h
--- a/Engine/Plugins/VirtualProduction/RemoteControl/Source/RemoteControlLogic/Public/Behaviour/Builtin/RangeMap/RCRangeMapBehaviour.h
+++ b/Engine/Plugins/VirtualProduction/RemoteControl/Source/RemoteControlLogic/Public/Behaviour/Builtin/RangeMap/RCRangeMapBehaviour.h
@@ -50,7 +50,7 @@ public:
 
 	/** Container holding all FRCRangeMapSteps correlating to each Action in the ActionContainer. */
 	UPROPERTY()
-	TMap<URCAction*, FRCRangeMapStep> RangeMapActionContainer;
+	TMap<TObjectPtr<URCAction>, FRCRangeMapStep> RangeMapActionContainer;
 
 public:
 	URCRangeMapBehaviour();
diff --git a/Engine/Plugins/VirtualProduction/RemoteControl/Source/RemoteControlLogic/Public/Behaviour/RCSetAssetByPathBehaviour.h b/Engine/Plugins/VirtualProduction/RemoteControl/Source/RemoteControlLogic/Public/Behaviour/RCSetAssetByPathBehaviour.h
--- a/Engine/Plugins/VirtualProduction/RemoteControl/Source/RemoteControlLogic/Public/Behaviour/RCSetAssetByPathBehaviour.h
+++ b/Engine/Plugins/VirtualProduction/RemoteControl/Source/RemoteControlLogic/Public/Behaviour/RCSetAssetByPathBehaviour.h
@@ -53,7 +53,7 @@ public:
 	TObjectPtr<URCVirtualPropertyContainerBase> PropertyInContainer;
 
 	UPROPERTY()
-	UClass* AssetClass;
+	TObjectPtr<UClass> AssetClass;
 
 	UPROPERTY()
 	FRCSetAssetPath PathStruct;
diff --git a/Engine/Plugins/VirtualProduction/RemoteControl/Source/WebRemoteControl/Private/RemoteControlReflectionUtils.h b/Engine/Plugins/VirtualProduction/RemoteControl/Source/WebRemoteControl/Private/RemoteControlReflectionUtils.h
--- a/Engine/Plugins/VirtualProduction/RemoteControl/Source/WebRemoteControl/Private/RemoteControlReflectionUtils.h
+++ b/Engine/Plugins/VirtualProduction/RemoteControl/Source/WebRemoteControl/Private/RemoteControlReflectionUtils.h
@@ -48,7 +48,7 @@ private:
 private:
 	/** Map of cached structs. */
 	UPROPERTY(Transient)
-	TMap<FName, UScriptStruct*> CachedStructs;
+	TMap<FName, TObjectPtr<UScriptStruct>> CachedStructs;
 };
 
 namespace UE
diff --git a/Engine/Plugins/VirtualProduction/RemoteControl/Source/WebRemoteControl/Private/RemoteControlResponse.h b/Engine/Plugins/VirtualProduction/RemoteControl/Source/WebRemoteControl/Private/RemoteControlResponse.h
--- a/Engine/Plugins/VirtualProduction/RemoteControl/Source/WebRemoteControl/Private/RemoteControlResponse.h
+++ b/Engine/Plugins/VirtualProduction/RemoteControl/Source/WebRemoteControl/Private/RemoteControlResponse.h
@@ -123,7 +123,7 @@ struct FDescribeObjectResponse
 	FString Name;
 
 	UPROPERTY()
-	UClass* Class = nullptr;
+	TObjectPtr<UClass> Class = nullptr;
 
 	UPROPERTY()
 	TArray<FRCPropertyDescription> Properties;
diff --git a/Engine/Plugins/VirtualProduction/StageMonitoring/Source/StageMonitorEditor/Private/Widgets/SDataProviderActivityFilter.h b/Engine/Plugins/VirtualProduction/StageMonitoring/Source/StageMonitorEditor/Private/Widgets/SDataProviderActivityFilter.h
--- a/Engine/Plugins/VirtualProduction/StageMonitoring/Source/StageMonitorEditor/Private/Widgets/SDataProviderActivityFilter.h
+++ b/Engine/Plugins/VirtualProduction/StageMonitoring/Source/StageMonitorEditor/Private/Widgets/SDataProviderActivityFilter.h
@@ -31,11 +31,11 @@ struct FDataProviderActivityFilterSettings
 	 * Used to detect new types that would be filtered out automatically
 	 */
 	UPROPERTY()
-	TArray<UScriptStruct*> ExistingPeriodicTypes;
+	TArray<TObjectPtr<UScriptStruct>> ExistingPeriodicTypes;
 
 	/** Message types that are filtered */
 	UPROPERTY()
-	TArray<UScriptStruct*> RestrictedTypes;
+	TArray<TObjectPtr<UScriptStruct>> RestrictedTypes;
 
 	/** Providers that are filtered using their friendly name */
 	UPROPERTY()
diff --git a/Engine/Plugins/VirtualProduction/TextureShare/Source/TextureShare/Private/Game/WorldSubsystem/TextureShareWorldSubsystem.h b/Engine/Plugins/VirtualProduction/TextureShare/Source/TextureShare/Private/Game/WorldSubsystem/TextureShareWorldSubsystem.h
--- a/Engine/Plugins/VirtualProduction/TextureShare/Source/TextureShare/Private/Game/WorldSubsystem/TextureShareWorldSubsystem.h
+++ b/Engine/Plugins/VirtualProduction/TextureShare/Source/TextureShare/Private/Game/WorldSubsystem/TextureShareWorldSubsystem.h
@@ -53,7 +53,7 @@ public:
 public:
 	// This UObject implements configuration and API
 	UPROPERTY(Transient)
-	UTextureShare* TextureShare;
+	TObjectPtr<UTextureShare> TextureShare;
 
 private:
 	// Release objects and resources
diff --git a/Engine/Plugins/VirtualProduction/TextureShare/Source/TextureShare/Public/Blueprints/TextureShareBlueprintContainers.h b/Engine/Plugins/VirtualProduction/TextureShare/Source/TextureShare/Public/Blueprints/TextureShareBlueprintContainers.h
--- a/Engine/Plugins/VirtualProduction/TextureShare/Source/TextureShare/Public/Blueprints/TextureShareBlueprintContainers.h
+++ b/Engine/Plugins/VirtualProduction/TextureShare/Source/TextureShare/Public/Blueprints/TextureShareBlueprintContainers.h
@@ -88,5 +88,5 @@ public:
 
 	// Objects for sharing
 	UPROPERTY(Transient)
-	TArray<UTextureShareObject*> TextureShareObjects;
+	TArray<TObjectPtr<UTextureShareObject>> TextureShareObjects;
 };
diff --git a/Engine/Plugins/VirtualProduction/TextureShare/Source/TextureShare/Public/Blueprints/TextureShareBlueprintContainersBase.h b/Engine/Plugins/VirtualProduction/TextureShare/Source/TextureShare/Public/Blueprints/TextureShareBlueprintContainersBase.h
--- a/Engine/Plugins/VirtualProduction/TextureShare/Source/TextureShare/Public/Blueprints/TextureShareBlueprintContainersBase.h
+++ b/Engine/Plugins/VirtualProduction/TextureShare/Source/TextureShare/Public/Blueprints/TextureShareBlueprintContainersBase.h
@@ -21,7 +21,7 @@ public:
 
 	// Resource to send
 	UPROPERTY(EditAnywhere, BlueprintReadWrite, Category = "TextureShare")
-	UTexture* Texture;
+	TObjectPtr<UTexture> Texture;
 };
 
 /**
@@ -39,7 +39,7 @@ public:
 
 	// Resource to receive
 	UPROPERTY(EditAnywhere, BlueprintReadWrite, Category = "TextureShare")
-	UTextureRenderTarget* Texture;
+	TObjectPtr<UTextureRenderTarget> Texture;
 };
 
 /**
