Log more info to help debug issue on partner stream

diff --git a/Engine/Plugins/VirtualProduction/LevelSnapshots/Source/LevelSnapshots/Private/Archive/ApplySnapshotToEditorArchive.cpp b/Engine/Plugins/VirtualProduction/LevelSnapshots/Source/LevelSnapshots/Private/Archive/ApplySnapshotToEditorArchive.cpp
--- a/Engine/Plugins/VirtualProduction/LevelSnapshots/Source/LevelSnapshots/Private/Archive/ApplySnapshotToEditorArchive.cpp
+++ b/Engine/Plugins/VirtualProduction/LevelSnapshots/Source/LevelSnapshots/Private/Archive/ApplySnapshotToEditorArchive.cpp
@@ -267,7 +267,11 @@ UE::LevelSnapshots::Private::FApplySnapshotToEditorArchive::FApplySnapshotToEdit
 		, SelectionSet(InSelectionSet)
 		, OnPropertySerialized(InOnPropertySerialized)
 		, Cache(Cache)
-{}
+{
+#if UE_BUILD_DEBUG
+	UE_LOG(LogLevelSnapshots, VeryVerbose, TEXT("FApplySnapshotToEditorArchive: %s (%s)"), *InOriginalObject->GetPathName(), *InOriginalObject->GetClass()->GetPathName());
+#endif
+}
 
 bool UE::LevelSnapshots::Private::FApplySnapshotToEditorArchive::ShouldSerializeAllProperties() const
 {
diff --git a/Engine/Plugins/VirtualProduction/LevelSnapshots/Source/LevelSnapshots/Private/Archive/ClassDefaults/ApplyClassDefaulDataArchive.cpp b/Engine/Plugins/VirtualProduction/LevelSnapshots/Source/LevelSnapshots/Private/Archive/ClassDefaults/ApplyClassDefaulDataArchive.cpp
--- a/Engine/Plugins/VirtualProduction/LevelSnapshots/Source/LevelSnapshots/Private/Archive/ClassDefaults/ApplyClassDefaulDataArchive.cpp
+++ b/Engine/Plugins/VirtualProduction/LevelSnapshots/Source/LevelSnapshots/Private/Archive/ClassDefaults/ApplyClassDefaulDataArchive.cpp
@@ -2,6 +2,7 @@
 
 #include "Archive/ClassDefaults/ApplyClassDefaulDataArchive.h"
 
+#include "LevelSnapshotsLog.h"
 #include "LevelSnapshotsModule.h"
 #include "ObjectSnapshotData.h"
 #include "WorldSnapshotData.h"
@@ -26,4 +27,8 @@ UE::LevelSnapshots::Private::FApplyClassDefaulDataArchive::FApplyClassDefaulData
 {
 	// Only overwrite transient properties when actually serializing into a snapshot CDO
 	ArSerializingDefaults = InSerialisationMode == ESerialisationMode::RestoringCDO;
+
+#if UE_BUILD_DEBUG
+	UE_LOG(LogLevelSnapshots, VeryVerbose, TEXT("FApplyClassDefaulDataArchive: %s (%s)"), *InObjectToRestore->GetPathName(), *InObjectToRestore->GetClass()->GetPathName());
+#endif
 }
diff --git a/Engine/Plugins/VirtualProduction/LevelSnapshots/Source/LevelSnapshots/Private/Archive/ClassDefaults/TakeClassDefaultObjectSnapshotArchive.cpp b/Engine/Plugins/VirtualProduction/LevelSnapshots/Source/LevelSnapshots/Private/Archive/ClassDefaults/TakeClassDefaultObjectSnapshotArchive.cpp
--- a/Engine/Plugins/VirtualProduction/LevelSnapshots/Source/LevelSnapshots/Private/Archive/ClassDefaults/TakeClassDefaultObjectSnapshotArchive.cpp
+++ b/Engine/Plugins/VirtualProduction/LevelSnapshots/Source/LevelSnapshots/Private/Archive/ClassDefaults/TakeClassDefaultObjectSnapshotArchive.cpp
@@ -2,6 +2,7 @@
 
 #include "Archive/ClassDefaults/TakeClassDefaultObjectSnapshotArchive.h"
 
+#include "LevelSnapshotsLog.h"
 #include "ObjectSnapshotData.h"
 #include "WorldSnapshotData.h"
 
@@ -17,4 +18,8 @@ void UE::LevelSnapshots::Private::FTakeClassDefaultObjectSnapshotArchive::SaveCl
 UE::LevelSnapshots::Private::FTakeClassDefaultObjectSnapshotArchive::FTakeClassDefaultObjectSnapshotArchive(FObjectSnapshotData& InObjectData, FWorldSnapshotData& InSharedData, UObject* InSerializedObject)
 	:
 	Super(InObjectData, InSharedData, false, InSerializedObject)
-{}
+{
+#if UE_BUILD_DEBUG
+	UE_LOG(LogLevelSnapshots, VeryVerbose, TEXT("FTakeClassDefaultObjectSnapshotArchive: %s (%s)"), *InSerializedObject->GetPathName(), *InSerializedObject->GetClass()->GetPathName());
+#endif
+}
diff --git a/Engine/Plugins/VirtualProduction/LevelSnapshots/Source/LevelSnapshots/Private/Archive/LoadSnapshotObjectArchive.cpp b/Engine/Plugins/VirtualProduction/LevelSnapshots/Source/LevelSnapshots/Private/Archive/LoadSnapshotObjectArchive.cpp
--- a/Engine/Plugins/VirtualProduction/LevelSnapshots/Source/LevelSnapshots/Private/Archive/LoadSnapshotObjectArchive.cpp
+++ b/Engine/Plugins/VirtualProduction/LevelSnapshots/Source/LevelSnapshots/Private/Archive/LoadSnapshotObjectArchive.cpp
@@ -67,4 +67,8 @@ UE::LevelSnapshots::Private::FLoadSnapshotObjectArchive::FLoadSnapshotObjectArch
 	: Super(InObjectData, InSharedData, true, InSerializedObject)
 	, ProcessObjectDependency(ProcessObjectDependency)
 	, Cache(Cache)
-{}
+{
+#if UE_BUILD_DEBUG
+	UE_LOG(LogLevelSnapshots, VeryVerbose, TEXT("FLoadSnapshotObjectArchive: %s (%s)"), *InSerializedObject->GetPathName(), *InSerializedObject->GetClass()->GetPathName());
+#endif
+}
diff --git a/Engine/Plugins/VirtualProduction/LevelSnapshots/Source/LevelSnapshots/Private/Archive/SnapshotArchive.cpp b/Engine/Plugins/VirtualProduction/LevelSnapshots/Source/LevelSnapshots/Private/Archive/SnapshotArchive.cpp
--- a/Engine/Plugins/VirtualProduction/LevelSnapshots/Source/LevelSnapshots/Private/Archive/SnapshotArchive.cpp
+++ b/Engine/Plugins/VirtualProduction/LevelSnapshots/Source/LevelSnapshots/Private/Archive/SnapshotArchive.cpp
@@ -41,11 +41,29 @@ bool UE::LevelSnapshots::Private::FSnapshotArchive::ShouldSkipProperty(const FPr
 {
 	// In debug builds only because this has big potential of impacting performance
 #if UE_BUILD_DEBUG
-	FString PropertyToDebug = UE::LevelSnapshots::ConsoleVariables::CVarBreakOnSerializedPropertyName.GetValueOnAnyThread();
+	FString PropertyToDebug = ConsoleVariables::CVarBreakOnSerializedPropertyName.GetValueOnAnyThread();
 	if (!PropertyToDebug.IsEmpty() && InProperty->GetName().Equals(PropertyToDebug, ESearchCase::IgnoreCase))
 	{
 		UE_DEBUG_BREAK();
 	}
+
+	const FString PropertyChain = [this, InProperty]()
+	{
+		if (GetSerializedPropertyChain())
+		{
+			FString Result;
+			for (int32 i = 0; i < GetSerializedPropertyChain()->GetNumProperties(); ++i)
+			{
+				Result += GetSerializedPropertyChain()->GetPropertyFromRoot(i)->GetAuthoredName();
+				Result += TEXT(".");
+			}
+			Result += InProperty->GetAuthoredName();
+			return Result;
+		}
+		
+		return InProperty->GetAuthoredName();
+	}();
+	UE_LOG(LogLevelSnapshots, VeryVerbose, TEXT("ShouldSkipProperty(%s %s)"), *PropertyChain, InProperty->IsNative() ? TEXT("C++") : TEXT("BP"));
 #endif
 	
 	const bool bIsPropertyUnsupported = InProperty->HasAnyPropertyFlags(ExcludedPropertyFlags);
diff --git a/Engine/Plugins/VirtualProduction/LevelSnapshots/Source/LevelSnapshots/Private/Archive/TakeWorldObjectSnapshotArchive.cpp b/Engine/Plugins/VirtualProduction/LevelSnapshots/Source/LevelSnapshots/Private/Archive/TakeWorldObjectSnapshotArchive.cpp
--- a/Engine/Plugins/VirtualProduction/LevelSnapshots/Source/LevelSnapshots/Private/Archive/TakeWorldObjectSnapshotArchive.cpp
+++ b/Engine/Plugins/VirtualProduction/LevelSnapshots/Source/LevelSnapshots/Private/Archive/TakeWorldObjectSnapshotArchive.cpp
@@ -2,6 +2,7 @@
 
 #include "TakeWorldObjectSnapshotArchive.h"
 
+#include "LevelSnapshotsLog.h"
 #include "ObjectSnapshotData.h"
 #include "WorldSnapshotData.h"
 
@@ -20,7 +21,11 @@ void UE::LevelSnapshots::Private::FTakeWorldObjectSnapshotArchive::TakeSnapshot(
 UE::LevelSnapshots::Private::FTakeWorldObjectSnapshotArchive::FTakeWorldObjectSnapshotArchive(FObjectSnapshotData& InObjectData, FWorldSnapshotData& InSharedData, UObject* InOriginalObject)
 	: Super(InObjectData, InSharedData, false, InOriginalObject)
 	, Archetype(InOriginalObject->GetArchetype())
-{}
+{
+#if UE_BUILD_DEBUG
+	UE_LOG(LogLevelSnapshots, VeryVerbose, TEXT("FTakeWorldObjectSnapshotArchive: %s (%s)"), *InOriginalObject->GetPathName(), *InOriginalObject->GetClass()->GetPathName());
+#endif
+}
 
 bool UE::LevelSnapshots::Private::FTakeWorldObjectSnapshotArchive::ShouldSkipProperty(const FProperty* InProperty) const
 {
