Remove master word

diff --git a/Engine/Plugins/VirtualProduction/LevelSnapshots/Source/LevelSnapshotsEditor/Private/Views/Filter/SLevelSnapshotsEditorFilters.h b/Engine/Plugins/VirtualProduction/LevelSnapshots/Source/LevelSnapshotsEditor/Private/Views/Filter/SLevelSnapshotsEditorFilters.h
--- a/Engine/Plugins/VirtualProduction/LevelSnapshots/Source/LevelSnapshotsEditor/Private/Views/Filter/SLevelSnapshotsEditorFilters.h
+++ b/Engine/Plugins/VirtualProduction/LevelSnapshots/Source/LevelSnapshotsEditor/Private/Views/Filter/SLevelSnapshotsEditorFilters.h
@@ -12,7 +12,6 @@ enum class EFilterChangeType : uint8;
 class FLevelSnapshotsEditorFilters;
 class SFavoriteFilterList;
 class SLevelSnapshotsEditorFilterRowGroup;
-class SMasterFilterIndicatorButton;
 class SCustomSplitter;
 class ULevelSnapshotsEditorData;
 class ULevelSnapshotFilter;
