Fix crash in certain cases when an object uses native struct serialization (e.g. UNiagaraComponent)

diff --git a/Engine/Plugins/VirtualProduction/LevelSnapshots/Source/LevelSnapshots/Private/Archive/ApplySnapshotToEditorArchive.cpp b/Engine/Plugins/VirtualProduction/LevelSnapshots/Source/LevelSnapshots/Private/Archive/ApplySnapshotToEditorArchive.cpp
--- a/Engine/Plugins/VirtualProduction/LevelSnapshots/Source/LevelSnapshots/Private/Archive/ApplySnapshotToEditorArchive.cpp
+++ b/Engine/Plugins/VirtualProduction/LevelSnapshots/Source/LevelSnapshots/Private/Archive/ApplySnapshotToEditorArchive.cpp
@@ -11,144 +11,8 @@
 #include "LevelSnapshotsLog.h"
 #include "Selection/PropertySelection.h"
 
-#include "Engine/World.h"
 #include "Serialization/ObjectReader.h"
-#include "Serialization/ObjectWriter.h"
-#include "UObject/TextProperty.h"
 #include "UObject/UnrealType.h"
-#include "Util/Property/WorldReferenceCheckingUtil.h"
-
-namespace UE::LevelSnapshots::Private::Internal
-{
-	using FShouldSerialize = TFunctionRef<bool(const FArchiveSerializedPropertyChain*, const FProperty*)>;
-	class FCopyProperties : public FObjectWriter
-	{
-		using Super = FObjectWriter;
-
-		FShouldSerialize PropertiesToSerialize;
-		UObject* SnapshotObject;
-
-		bool ContainsWorldObjectProperty(const FProperty* InProperty) const
-		{
-			return UE::LevelSnapshots::Private::ContainsSubobjectOrSatisfiesPredicate(SnapshotObject, GetSerializedPropertyChain(), InProperty, [](UObject* ContainedUObject)
-			{
-				return ContainedUObject && ContainedUObject->IsInA(UWorld::StaticClass());
-			});
-		}
-				
-	public:
-
-		mutable TSet<const FTextProperty*> TextProperties;
-
-		FCopyProperties(UObject* SnapshotObject, TArray<uint8>& SaveLocation, FShouldSerialize PropertiesToSerialize)
-			: Super(SaveLocation)
-			, PropertiesToSerialize(PropertiesToSerialize)
-			, SnapshotObject(SnapshotObject)
-		{
-			ArNoDelta = true;
-		}
-
-		virtual bool ShouldSkipProperty(const FProperty* InProperty) const override
-		{
-			// Do not copy object reference properties that have a world as outer: They will not be valid when copied from snapshot to editor world.
-			// Hence, we only allow object references to external assets, e.g. UMaterials or UDataAssets
-			if (ContainsWorldObjectProperty(InProperty))
-			{
-				return true;
-			}
-
-			const bool bIsPropertyAllowed = PropertiesToSerialize(GetSerializedPropertyChain(), InProperty);
-			if (bIsPropertyAllowed)
-			{
-				if (const FTextProperty* TextProperty = CastField<FTextProperty>(InProperty))
-				{
-					TextProperties.Add(TextProperty);
-					return false;
-				}
-			}
-			return !bIsPropertyAllowed;
-		}		
-	};
-
-	class FSerializeTextProperties : public FApplyClassDefaulDataArchive
-	{
-		TSet<const FTextProperty*>& TextProperties;
-	public:
-		
-		FSerializeTextProperties(TSet<const FTextProperty*>& InTextProperties, FObjectSnapshotData& InObjectData, FWorldSnapshotData& InSharedData, UObject* InSerializedObject)
-			:
-			FApplyClassDefaulDataArchive(InObjectData, InSharedData, InSerializedObject, ESerialisationMode::RestoringChangedDefaults),
-			TextProperties(InTextProperties)
-		{}
-
-		virtual bool ShouldSkipProperty(const FProperty* InProperty) const override
-		{
-			if (const FTextProperty* TextProperty = CastField<FTextProperty>(InProperty))
-			{
-				return !TextProperties.Contains(TextProperty);
-			}
-			return true;
-		}
-	};
-	
-	static TSet<const FTextProperty*> CopyPastePropertiesDifferentInCDO(FShouldSerialize ShouldSerialize, UObject* InOriginalObject, UObject* InDeserializedVersion)
-	{
-		TArray<uint8> CopiedPropertyData;
-		FCopyProperties CopySimpleProperties(InDeserializedVersion, CopiedPropertyData, ShouldSerialize);
-		InDeserializedVersion->Serialize(CopySimpleProperties);
-		
-		FObjectReader PasteSimpleProperties(InOriginalObject, CopiedPropertyData);
-		return CopySimpleProperties.TextProperties;
-	}
-
-	static TSet<const FTextProperty*> CopyPastePropertiesDifferentInCDO(const FPropertySelection& PropertiesLeftToSerialise, UObject* InOriginalObject, UObject* InDeserializedVersion)
-	{
-		if (PropertiesLeftToSerialise.IsEmpty())
-		{
-			return {};
-		}
-
-		return CopyPastePropertiesDifferentInCDO(
-			[&PropertiesLeftToSerialise](const FArchiveSerializedPropertyChain* Chain, const FProperty* Property)
-			{
-				return PropertiesLeftToSerialise.ShouldSerializeProperty(Chain, Property);
-			},
-			InOriginalObject,
-			InDeserializedVersion
-			);
-	}
-
-	static void FixUpTextPropertiesDifferentInCDO(
-		TSet<const FTextProperty*> TextProperties,
-		FWorldSnapshotData& InSharedData,
-		TOptional<FClassDataIndex> ClassIndex,
-		FSnapshotDataCache& Cache,
-		UObject* InOriginalObject)
-	{
-		const bool bAreTextPropertiesLeft = TextProperties.Num() > 0;
-		if (!bAreTextPropertiesLeft)
-		{
-			return;
-		}
-
-		if (ClassIndex)
-		{
-			const FSubobjectArchetypeFallbackInfo FallbackInfo { InOriginalObject->GetOuter() };
-			TOptional<TNonNullPtr<FClassSnapshotData>> ClassDefaults =
-				InOriginalObject->IsA<AActor>()
-					? TOptional<TNonNullPtr<FClassSnapshotData>>(&InSharedData.ClassData[*ClassIndex])
-					: GetSubobjectArchetypeData(InSharedData, *ClassIndex, Cache, FallbackInfo);
-			if (ClassDefaults)
-			{
-				FSerializeTextProperties SerializeTextProperties(TextProperties, *ClassDefaults.GetValue(), InSharedData, InOriginalObject);
-				InOriginalObject->Serialize(SerializeTextProperties);
-				return;
-			}
-		}
-			
-		UE_LOG(LogLevelSnapshots, Warning, TEXT("%d text properties have changed in class defaults since snapshot was taken but cannot be restored."), TextProperties.Num());
-	}
-}
 
 void UE::LevelSnapshots::Private::FApplySnapshotToEditorArchive::ApplyToExistingEditorWorldObject(
 	FObjectSnapshotData& InObjectData,
@@ -173,16 +37,17 @@ void UE::LevelSnapshots::Private::FApplySnapshotToEditorArchive::ApplyToExisting
 
 	FPropertySelection PropertiesLeftToSerialize = *Selection;
 	auto TrackSerializedProperties = [&PropertiesLeftToSerialize](const FArchiveSerializedPropertyChain* Chain, const FProperty* Property) { PropertiesLeftToSerialize.RemoveProperty(Chain, Property); };
+
+	// 1. Serialize archetype first to handle the case were the archetype has changed properties since the snapshot was taken
+	if (ClassIndex) // Sometimes not available, e.g. for custom subobjects
+	{
+		const FSubobjectArchetypeFallbackInfo ClassFallbackInfo{ InOriginalObject->GetOuter(), InOriginalObject->GetFName(), InOriginalObject->GetFlags() };
+		SerializeSelectedClassDefaultsInto(InOriginalObject, InSharedData, *ClassIndex, Cache, ClassFallbackInfo, *Selection);
+	}
 	
-	// Step 1: Serialise  properties that were different from CDO at time of snapshotting and that are still different from CDO
+	// Step 2: Serialise  properties that were different from CDO at time of snapshotting and that are still different from CDO
 	FApplySnapshotToEditorArchive ApplySavedData(InObjectData, InSharedData, InOriginalObject, InSelectionMapForResolvingSubobjects, Selection, TrackSerializedProperties, Cache);
 	InOriginalObject->Serialize(ApplySavedData);
-	
-	// Step 2: Serialise any remaining properties that were not covered: properties that were equal to the CDO value when the snapshot was taken but now are different from the CDO.
-	TSet<const FTextProperty*> TextProperties = Internal::CopyPastePropertiesDifferentInCDO(PropertiesLeftToSerialize, InOriginalObject, InDeserializedVersion);
-
-	// Step 3: Serialize FText properties that have changed in CDO since snapshot was taken. 
-	Internal::FixUpTextPropertiesDifferentInCDO(MoveTemp(TextProperties), InSharedData, ClassIndex, Cache, InOriginalObject);
 }
 
 void UE::LevelSnapshots::Private::FApplySnapshotToEditorArchive::ApplyToEditorWorldObjectRecreatedWithArchetype(
@@ -227,7 +92,6 @@ bool UE::LevelSnapshots::Private::FApplySnapshotToEditorArchive::ShouldSkipPrope
 	SCOPED_SNAPSHOT_CORE_TRACE(ShouldSkipProperty);
 	
 	bool bShouldSkipProperty = Super::ShouldSkipProperty(InProperty);
-	
 	if (!bShouldSkipProperty && !ShouldSerializeAllProperties())
 	{
 		const bool bIsAllowed = SelectionSet.GetValue()->ShouldSerializeProperty(GetSerializedPropertyChain(), InProperty);  
@@ -245,7 +109,7 @@ void UE::LevelSnapshots::Private::FApplySnapshotToEditorArchive::PushSerializedP
 	Super::PushSerializedProperty(InProperty, bIsEditorOnlyProperty);
 }
 
-UObject* UE::LevelSnapshots::Private::FApplySnapshotToEditorArchive::ResolveObjectDependency(int32 ObjectIndex) const
+UObject* UE::LevelSnapshots::Private::FApplySnapshotToEditorArchive::ResolveObjectDependency(int32 ObjectIndex, UObject* CurrentValue) const
 {
 	FString LocalizationNamespace;
 #if USE_STABLE_LOCALIZATION_KEYS
diff --git a/Engine/Plugins/VirtualProduction/LevelSnapshots/Source/LevelSnapshots/Private/Archive/ApplySnapshotToEditorArchive.h b/Engine/Plugins/VirtualProduction/LevelSnapshots/Source/LevelSnapshots/Private/Archive/ApplySnapshotToEditorArchive.h
--- a/Engine/Plugins/VirtualProduction/LevelSnapshots/Source/LevelSnapshots/Private/Archive/ApplySnapshotToEditorArchive.h
+++ b/Engine/Plugins/VirtualProduction/LevelSnapshots/Source/LevelSnapshots/Private/Archive/ApplySnapshotToEditorArchive.h
@@ -65,7 +65,7 @@ namespace UE::LevelSnapshots::Private
 	protected:
 
 		//~ Begin FSnapshotArchive Interface
-		virtual UObject* ResolveObjectDependency(int32 ObjectIndex) const override;
+		virtual UObject* ResolveObjectDependency(int32 ObjectIndex, UObject* CurrentValue) const override;
 		//~ End FSnapshotArchive Interface
 
 	private:
diff --git a/Engine/Plugins/VirtualProduction/LevelSnapshots/Source/LevelSnapshots/Private/Archive/ClassDefaults/ApplyClassDefaulDataArchive.cpp b/Engine/Plugins/VirtualProduction/LevelSnapshots/Source/LevelSnapshots/Private/Archive/ClassDefaults/ApplyClassDefaulDataArchive.cpp
--- a/Engine/Plugins/VirtualProduction/LevelSnapshots/Source/LevelSnapshots/Private/Archive/ClassDefaults/ApplyClassDefaulDataArchive.cpp
+++ b/Engine/Plugins/VirtualProduction/LevelSnapshots/Source/LevelSnapshots/Private/Archive/ClassDefaults/ApplyClassDefaulDataArchive.cpp
@@ -5,7 +5,9 @@
 #include "LevelSnapshotsLog.h"
 #include "LevelSnapshotsModule.h"
 #include "ObjectSnapshotData.h"
+#include "SelectionSet.h"
 #include "WorldSnapshotData.h"
+#include "Util/WorldData/SnapshotObjectUtil.h"
 
 void UE::LevelSnapshots::Private::FApplyClassDefaulDataArchive::SerializeClassDefaultObject(FObjectSnapshotData& InObjectData, FWorldSnapshotData& InSharedData, UObject* InCDO)
 {
@@ -21,9 +23,33 @@ void UE::LevelSnapshots::Private::FApplyClassDefaulDataArchive::RestoreChangedCl
 	InObjectToRestore->Serialize(Archive);
 }
 
-UE::LevelSnapshots::Private::FApplyClassDefaulDataArchive::FApplyClassDefaulDataArchive(FObjectSnapshotData& InObjectData, FWorldSnapshotData& InSharedData, UObject* InObjectToRestore, ESerialisationMode InSerialisationMode)
-	:
-	Super(InObjectData, InSharedData, true, InObjectToRestore)
+void UE::LevelSnapshots::Private::FApplyClassDefaulDataArchive::RestoreSelectedChangedClassDefaults(FObjectSnapshotData& InObjectData, FWorldSnapshotData& InSharedData, UObject* InObjectToRestore, const FPropertySelection& PropertiesToRestore)
+{
+	FApplyClassDefaulDataArchive Archive(InObjectData, InSharedData, InObjectToRestore, ESerialisationMode::RestoringChangedDefaults, &PropertiesToRestore);
+	InObjectToRestore->Serialize(Archive);
+}
+
+bool UE::LevelSnapshots::Private::FApplyClassDefaulDataArchive::ShouldSkipProperty(const FProperty* InProperty) const
+{
+	bool bShouldSkipProperty = Super::ShouldSkipProperty(InProperty);
+	const bool bShouldSerializeAllProperties = !SelectionSet.IsSet();
+	if (!bShouldSkipProperty && !bShouldSerializeAllProperties)
+	{
+		const bool bIsAllowed = SelectionSet.GetValue()->ShouldSerializeProperty(GetSerializedPropertyChain(), InProperty);  
+		bShouldSkipProperty = !bIsAllowed;
+	}
+	
+	return bShouldSkipProperty;
+}
+
+UE::LevelSnapshots::Private::FApplyClassDefaulDataArchive::FApplyClassDefaulDataArchive(
+	FObjectSnapshotData& InObjectData,
+	FWorldSnapshotData& InSharedData,
+	UObject* InObjectToRestore,
+	ESerialisationMode InSerialisationMode,
+	TOptional<const FPropertySelection*> InSelectionSet)
+	: Super(InObjectData, InSharedData, true, InObjectToRestore)
+	, SelectionSet(InSelectionSet)
 {
 	// Only overwrite transient properties when actually serializing into a snapshot CDO
 	ArSerializingDefaults = InSerialisationMode == ESerialisationMode::RestoringCDO;
@@ -32,3 +58,8 @@ UE::LevelSnapshots::Private::FApplyClassDefaulDataArchive::FApplyClassDefaulData
 	UE_LOG(LogLevelSnapshots, VeryVerbose, TEXT("FApplyClassDefaulDataArchive: %s (%s)"), *InObjectToRestore->GetPathName(), *InObjectToRestore->GetClass()->GetPathName());
 #endif
 }
+
+UObject* UE::LevelSnapshots::Private::FApplyClassDefaulDataArchive::ResolveObjectDependency(int32 ObjectIndex, UObject* CurrentValue) const
+{
+	return ResolveObjectDependencyForClassDefaultObject(GetSharedData(), ObjectIndex, CurrentValue);
+}
diff --git a/Engine/Plugins/VirtualProduction/LevelSnapshots/Source/LevelSnapshots/Private/Archive/ClassDefaults/ApplyClassDefaulDataArchive.h b/Engine/Plugins/VirtualProduction/LevelSnapshots/Source/LevelSnapshots/Private/Archive/ClassDefaults/ApplyClassDefaulDataArchive.h
--- a/Engine/Plugins/VirtualProduction/LevelSnapshots/Source/LevelSnapshots/Private/Archive/ClassDefaults/ApplyClassDefaulDataArchive.h
+++ b/Engine/Plugins/VirtualProduction/LevelSnapshots/Source/LevelSnapshots/Private/Archive/ClassDefaults/ApplyClassDefaulDataArchive.h
@@ -3,11 +3,13 @@
 #pragma once
 
 #include "CoreMinimal.h"
+#include "WorldSnapshotData.h"
 #include "Archive/ClassDefaults/BaseClassDefaultArchive.h"
 
+class UObject;
 struct FObjectSnapshotData;
+struct FPropertySelection;
 struct FWorldSnapshotData;
-class UObject;
 
 namespace UE::LevelSnapshots::Private
 {
@@ -21,7 +23,12 @@ namespace UE::LevelSnapshots::Private
 		static void SerializeClassDefaultObject(FObjectSnapshotData& InObjectData, FWorldSnapshotData& InSharedData, UObject* InCDO);
 		/* Archive will be used to to apply any non-object properties before the object receives its saved snapshot values. This is to handle when a CDO has changed its default values since a snapshot was taken. */
 		static void RestoreChangedClassDefaults(FObjectSnapshotData& InObjectData, FWorldSnapshotData& InSharedData, UObject* ObjectToRestore);
+		static void RestoreSelectedChangedClassDefaults(FObjectSnapshotData& InObjectData, FWorldSnapshotData& InSharedData, UObject* InObjectToRestore, const FPropertySelection& PropertiesToRestore);
 
+		//~ Begin FArchive Interface
+		virtual bool ShouldSkipProperty(const FProperty* InProperty) const override;
+		//~ End FArchive Interface
+		
 	protected:
 	
 		enum class ESerialisationMode
@@ -30,6 +37,13 @@ namespace UE::LevelSnapshots::Private
 			RestoringChangedDefaults
 		};
 
-		FApplyClassDefaulDataArchive(FObjectSnapshotData& InObjectData, FWorldSnapshotData& InSharedData, UObject* InObjectToRestore, ESerialisationMode InSerialisationMode);
+		/** Immutable list of properties we are supposed to serialise. Serialize all properties if unset. */
+		TOptional<const FPropertySelection*> SelectionSet;
+
+		FApplyClassDefaulDataArchive(FObjectSnapshotData& InObjectData, FWorldSnapshotData& InSharedData, UObject* InObjectToRestore, ESerialisationMode InSerialisationMode, TOptional<const FPropertySelection*> InSelectionSet = {});
+
+		//~ Begin FSnapshotArchive Interface
+		virtual UObject* ResolveObjectDependency(int32 ObjectIndex, UObject* CurrentValue) const override;
+		//~ End FSnapshotArchive Interface
 	};
 }
\ No newline at end of file
diff --git a/Engine/Plugins/VirtualProduction/LevelSnapshots/Source/LevelSnapshots/Private/Archive/ClassDefaults/BaseClassDefaultArchive.cpp b/Engine/Plugins/VirtualProduction/LevelSnapshots/Source/LevelSnapshots/Private/Archive/ClassDefaults/BaseClassDefaultArchive.cpp
--- a/Engine/Plugins/VirtualProduction/LevelSnapshots/Source/LevelSnapshots/Private/Archive/ClassDefaults/BaseClassDefaultArchive.cpp
+++ b/Engine/Plugins/VirtualProduction/LevelSnapshots/Source/LevelSnapshots/Private/Archive/ClassDefaults/BaseClassDefaultArchive.cpp
@@ -3,24 +3,8 @@
 #include "BaseClassDefaultArchive.h"
 
 #include "Data/Util/Property/PropertyUtil.h"
-#include "Data/Util/WorldData/SnapshotObjectUtil.h"
 #include "WorldSnapshotData.h"
 
-#include "Util/Property/WorldReferenceCheckingUtil.h"
-
-bool UE::LevelSnapshots::Private::FBaseClassDefaultArchive::ShouldSkipProperty(const FProperty* InProperty) const
-{
-	const bool bSuperWantsToSkip = Super::ShouldSkipProperty(InProperty);
-	
-	// Avoid any object references owned by a CDO. References are automatically replaced by object instance creation flow.
-	return bSuperWantsToSkip || IsPropertyReferenceToSubobjectOrClassDefaults(InProperty);
-}
-
-UObject* UE::LevelSnapshots::Private::FBaseClassDefaultArchive::ResolveObjectDependency(int32 ObjectIndex) const
-{
-	return ResolveObjectDependencyForClassDefaultObject(GetSharedData(), ObjectIndex);
-}
-
 UE::LevelSnapshots::Private::FBaseClassDefaultArchive::FBaseClassDefaultArchive(FObjectSnapshotData& InObjectData, FWorldSnapshotData& InSharedData, bool bIsLoading, UObject* InObjectToRestore)
 	:
 	Super(InObjectData, InSharedData, bIsLoading, InObjectToRestore)
@@ -32,20 +16,4 @@ UE::LevelSnapshots::Private::FBaseClassDefaultArchive::FBaseClassDefaultArchive(
 	
 	// Otherwise we are not allowed to serialize transient properties
 	ArSerializingDefaults = true;
-}
-
-bool UE::LevelSnapshots::Private::FBaseClassDefaultArchive::IsPropertyReferenceToSubobjectOrClassDefaults(const FProperty* InProperty) const
-{
-	return UE::LevelSnapshots::Private::ContainsSubobjectOrSatisfiesPredicate(GetSerializedObject(), GetSerializedPropertyChain(), InProperty, [this](UObject* ContainedPtr)
-	{
-		if (ContainedPtr)
-		{
-			const bool bIsClassDefault = ContainedPtr->HasAnyFlags(RF_ClassDefaultObject | RF_ArchetypeObject);
-			const bool bIsPointingToDefaultSubobject = ContainedPtr->HasAnyFlags(RF_DefaultSubObject);
-			const bool bIsPointingToSelf = ContainedPtr == GetSerializedObject();
-			const bool bIsPointingToSubobject = ContainedPtr->IsIn(GetSerializedObject());
-			return bIsClassDefault || bIsPointingToDefaultSubobject || bIsPointingToSelf || bIsPointingToSubobject;
-		}
-		return false;
-	}); 
 }
\ No newline at end of file
diff --git a/Engine/Plugins/VirtualProduction/LevelSnapshots/Source/LevelSnapshots/Private/Archive/ClassDefaults/BaseClassDefaultArchive.h b/Engine/Plugins/VirtualProduction/LevelSnapshots/Source/LevelSnapshots/Private/Archive/ClassDefaults/BaseClassDefaultArchive.h
--- a/Engine/Plugins/VirtualProduction/LevelSnapshots/Source/LevelSnapshots/Private/Archive/ClassDefaults/BaseClassDefaultArchive.h
+++ b/Engine/Plugins/VirtualProduction/LevelSnapshots/Source/LevelSnapshots/Private/Archive/ClassDefaults/BaseClassDefaultArchive.h
@@ -11,26 +11,12 @@ struct FWorldSnapshotData;
 
 namespace UE::LevelSnapshots::Private
 {
-	/* Shared logic for serializing class defaults. */
+	/** Shared logic for serializing class defaults. */
 	class FBaseClassDefaultArchive : public FSnapshotArchive
 	{
 		using Super = FSnapshotArchive;
-	public:
-
-		//~ Begin FSnapshotArchive Interface
-		virtual bool ShouldSkipProperty(const FProperty* InProperty) const override;
-		//~ End FSnapshotArchive Interface
-
-		protected:
-	
-		//~ Begin FSnapshotArchive Interface
-		virtual UObject* ResolveObjectDependency(int32 ObjectIndex) const override;
-		//~ End FSnapshotArchive Interface
+	protected:
 	
 		FBaseClassDefaultArchive(FObjectSnapshotData& InObjectData, FWorldSnapshotData& InSharedData, bool bIsLoading, UObject* InObjectToRestore);
-
-	private:
-	
-		bool IsPropertyReferenceToSubobjectOrClassDefaults(const FProperty* InProperty) const;
 	};
 }
\ No newline at end of file
diff --git a/Engine/Plugins/VirtualProduction/LevelSnapshots/Source/LevelSnapshots/Private/Archive/ClassDefaults/TakeClassDefaultObjectSnapshotArchive.cpp b/Engine/Plugins/VirtualProduction/LevelSnapshots/Source/LevelSnapshots/Private/Archive/ClassDefaults/TakeClassDefaultObjectSnapshotArchive.cpp
--- a/Engine/Plugins/VirtualProduction/LevelSnapshots/Source/LevelSnapshots/Private/Archive/ClassDefaults/TakeClassDefaultObjectSnapshotArchive.cpp
+++ b/Engine/Plugins/VirtualProduction/LevelSnapshots/Source/LevelSnapshots/Private/Archive/ClassDefaults/TakeClassDefaultObjectSnapshotArchive.cpp
@@ -15,6 +15,24 @@ void UE::LevelSnapshots::Private::FTakeClassDefaultObjectSnapshotArchive::SaveCl
 	InObjectData.ClassFlags = SerializedObject->GetClass()->GetFlags();
 }
 
+void UE::LevelSnapshots::Private::FTakeClassDefaultObjectSnapshotArchive::OnAddObjectDependency(int32 ObjectIndex, UObject* Object) const
+{
+	if (Object)
+	{
+		const bool bIsClassDefault = Object->HasAnyFlags(RF_ClassDefaultObject | RF_ArchetypeObject);
+		const bool bIsPointingToDefaultSubobject = Object->HasAnyFlags(RF_DefaultSubObject);
+		const bool bIsPointingToSelf = Object == GetSerializedObject();
+		const bool bIsPointingToSubobject = Object->IsIn(GetSerializedObject());
+		const bool bShouldSkip = bIsClassDefault || bIsPointingToDefaultSubobject || bIsPointingToSelf || bIsPointingToSubobject;
+
+		// We don't ever want to interfere with the default subobjects that the class archetype assigns so the reference is marked so we know it is supposed to be skipped when restoring 
+		if (bShouldSkip)
+		{
+			GetSharedData().SerializedReferenceMetaData.FindOrAdd(ObjectIndex).Flags |= EObjectReferenceSnapshotFlags::SkipWhenSerializingArchetypeData;
+		}
+	}
+}
+
 UE::LevelSnapshots::Private::FTakeClassDefaultObjectSnapshotArchive::FTakeClassDefaultObjectSnapshotArchive(FObjectSnapshotData& InObjectData, FWorldSnapshotData& InSharedData, UObject* InSerializedObject)
 	:
 	Super(InObjectData, InSharedData, false, InSerializedObject)
diff --git a/Engine/Plugins/VirtualProduction/LevelSnapshots/Source/LevelSnapshots/Private/Archive/ClassDefaults/TakeClassDefaultObjectSnapshotArchive.h b/Engine/Plugins/VirtualProduction/LevelSnapshots/Source/LevelSnapshots/Private/Archive/ClassDefaults/TakeClassDefaultObjectSnapshotArchive.h
--- a/Engine/Plugins/VirtualProduction/LevelSnapshots/Source/LevelSnapshots/Private/Archive/ClassDefaults/TakeClassDefaultObjectSnapshotArchive.h
+++ b/Engine/Plugins/VirtualProduction/LevelSnapshots/Source/LevelSnapshots/Private/Archive/ClassDefaults/TakeClassDefaultObjectSnapshotArchive.h
@@ -20,6 +20,13 @@ namespace UE::LevelSnapshots::Private
 
 		static void SaveClassDefaultObject(FClassSnapshotData& InObjectData, FWorldSnapshotData& InSharedData, UObject* SerializedObject);
 
+	protected:
+
+		//~ Begin FSnapshotArchive Interface
+		virtual UObject* ResolveObjectDependency(int32 ObjectIndex, UObject* CurrentValue) const override { checkNoEntry(); return nullptr; }
+		virtual void OnAddObjectDependency(int32 ObjectIndex, UObject* Object) const override;
+		//~ End FSnapshotArchive Interface
+		
 	private:
 	
 		FTakeClassDefaultObjectSnapshotArchive(FObjectSnapshotData& InObjectData, FWorldSnapshotData& InSharedData, UObject* InSerializedObject);
diff --git a/Engine/Plugins/VirtualProduction/LevelSnapshots/Source/LevelSnapshots/Private/Archive/LoadSnapshotObjectArchive.cpp b/Engine/Plugins/VirtualProduction/LevelSnapshots/Source/LevelSnapshots/Private/Archive/LoadSnapshotObjectArchive.cpp
--- a/Engine/Plugins/VirtualProduction/LevelSnapshots/Source/LevelSnapshots/Private/Archive/LoadSnapshotObjectArchive.cpp
+++ b/Engine/Plugins/VirtualProduction/LevelSnapshots/Source/LevelSnapshots/Private/Archive/LoadSnapshotObjectArchive.cpp
@@ -52,7 +52,7 @@ void UE::LevelSnapshots::Private::FLoadSnapshotObjectArchive::ApplyToSnapshotWor
 	FLevelSnapshotsModule::GetInternalModuleInstance().OnPostLoadSnapshotObject({ InObjectToRestore, InSharedData });
 }
 
-UObject* UE::LevelSnapshots::Private::FLoadSnapshotObjectArchive::ResolveObjectDependency(int32 ObjectIndex) const
+UObject* UE::LevelSnapshots::Private::FLoadSnapshotObjectArchive::ResolveObjectDependency(int32 ObjectIndex, UObject* CurrentValue) const
 {
 	FString LocalizationNamespace;
 #if USE_STABLE_LOCALIZATION_KEYS
diff --git a/Engine/Plugins/VirtualProduction/LevelSnapshots/Source/LevelSnapshots/Private/Archive/LoadSnapshotObjectArchive.h b/Engine/Plugins/VirtualProduction/LevelSnapshots/Source/LevelSnapshots/Private/Archive/LoadSnapshotObjectArchive.h
--- a/Engine/Plugins/VirtualProduction/LevelSnapshots/Source/LevelSnapshots/Private/Archive/LoadSnapshotObjectArchive.h
+++ b/Engine/Plugins/VirtualProduction/LevelSnapshots/Source/LevelSnapshots/Private/Archive/LoadSnapshotObjectArchive.h
@@ -22,7 +22,7 @@ namespace UE::LevelSnapshots::Private
 	protected:
 
 		//~ Begin FSnapshotArchive Interface
-		virtual UObject* ResolveObjectDependency(int32 ObjectIndex) const override;
+		virtual UObject* ResolveObjectDependency(int32 ObjectIndex, UObject* CurrentValue) const override;
 		//~ End FSnapshotArchive Interface
 
 	private:
diff --git a/Engine/Plugins/VirtualProduction/LevelSnapshots/Source/LevelSnapshots/Private/Archive/SnapshotArchive.cpp b/Engine/Plugins/VirtualProduction/LevelSnapshots/Source/LevelSnapshots/Private/Archive/SnapshotArchive.cpp
--- a/Engine/Plugins/VirtualProduction/LevelSnapshots/Source/LevelSnapshots/Private/Archive/SnapshotArchive.cpp
+++ b/Engine/Plugins/VirtualProduction/LevelSnapshots/Source/LevelSnapshots/Private/Archive/SnapshotArchive.cpp
@@ -123,7 +123,7 @@ FArchive& UE::LevelSnapshots::Private::FSnapshotArchive::operator<<(UObject*& Va
 			return *this;
 		}
 
-		Value = ResolveObjectDependency(ReferencedIndex);
+		Value = ResolveObjectDependency(ReferencedIndex, Value);
 		if (Value)
 		{
 			const ULevel* ResolvedOwningLevel = Value->GetTypedOuter<ULevel>();
@@ -140,7 +140,8 @@ FArchive& UE::LevelSnapshots::Private::FSnapshotArchive::operator<<(UObject*& Va
 	}
 	else
 	{
-		int32 ReferenceIndex = UE::LevelSnapshots::Private::AddObjectDependency(SharedData, Value);
+		int32 ReferenceIndex = AddObjectDependency(SharedData, Value);
+		OnAddObjectDependency(ReferenceIndex, Value);
 		*this << ReferenceIndex;
 	}
 	
diff --git a/Engine/Plugins/VirtualProduction/LevelSnapshots/Source/LevelSnapshots/Private/Archive/SnapshotArchive.h b/Engine/Plugins/VirtualProduction/LevelSnapshots/Source/LevelSnapshots/Private/Archive/SnapshotArchive.h
--- a/Engine/Plugins/VirtualProduction/LevelSnapshots/Source/LevelSnapshots/Private/Archive/SnapshotArchive.h
+++ b/Engine/Plugins/VirtualProduction/LevelSnapshots/Source/LevelSnapshots/Private/Archive/SnapshotArchive.h
@@ -26,16 +26,17 @@ namespace UE::LevelSnapshots::Private
 		virtual int64 Tell() override;
 		virtual void Seek(int64 InPos) override;
 		virtual bool ShouldSkipProperty(const FProperty* InProperty) const override;
-	
 		virtual FArchive& operator<<(FName& Value) override;
 		virtual FArchive& operator<<(UObject*& Value) override;
 		virtual void Serialize(void* Data, int64 Length) override;
 		//~ End FArchive Interface
 
-		protected:
+	protected:
 	
 		/* Allocates and serializes an object dependency, or gets the object, if it already exists. */
-		virtual UObject* ResolveObjectDependency(int32 ObjectIndex) const = 0;
+		virtual UObject* ResolveObjectDependency(int32 ObjectIndex, UObject* CurrentValue) const = 0;
+		/** Called when an object dependency is saved. */
+		virtual void OnAddObjectDependency(int32 ObjectIndex, UObject* Object) const {}
 	
 		FWorldSnapshotData& GetSharedData() const { return SharedData;}
 		UObject* GetSerializedObject() const { return SerializedObject; }
diff --git a/Engine/Plugins/VirtualProduction/LevelSnapshots/Source/LevelSnapshots/Private/Archive/TakeWorldObjectSnapshotArchive.cpp b/Engine/Plugins/VirtualProduction/LevelSnapshots/Source/LevelSnapshots/Private/Archive/TakeWorldObjectSnapshotArchive.cpp
--- a/Engine/Plugins/VirtualProduction/LevelSnapshots/Source/LevelSnapshots/Private/Archive/TakeWorldObjectSnapshotArchive.cpp
+++ b/Engine/Plugins/VirtualProduction/LevelSnapshots/Source/LevelSnapshots/Private/Archive/TakeWorldObjectSnapshotArchive.cpp
@@ -42,15 +42,9 @@ bool UE::LevelSnapshots::Private::FTakeWorldObjectSnapshotArchive::ShouldSkipPro
 		{
 			// We are within a struct property and we have already checked that the struct property does not have equal values
 			// Simply allow all properties but only ones which are not deprecated.
-			// TODO: We could save more disk space here by checking this property but let's test this first...
+			// TODO: We could save more disk space here by checking the subproperties
 			return InProperty->HasAnyPropertyFlags(CPF_Deprecated | CPF_Transient);
 		}
-
-		// Always save object properties regardless whether different from CDO or not ... this makes restoring easier: see FApplySnapshotDataArchiveV2::ApplyToExistingWorldObject
-		if (const FObjectPropertyBase* ObjectProperty = CastField<FObjectPropertyBase>(InProperty))
-		{
-			return false;
-		}
 		
 		UObject* OriginalContainer = GetSerializedObject();
 		UObject* ClassDefaultContainer = Archetype;
diff --git a/Engine/Plugins/VirtualProduction/LevelSnapshots/Source/LevelSnapshots/Private/Archive/TakeWorldObjectSnapshotArchive.h b/Engine/Plugins/VirtualProduction/LevelSnapshots/Source/LevelSnapshots/Private/Archive/TakeWorldObjectSnapshotArchive.h
--- a/Engine/Plugins/VirtualProduction/LevelSnapshots/Source/LevelSnapshots/Private/Archive/TakeWorldObjectSnapshotArchive.h
+++ b/Engine/Plugins/VirtualProduction/LevelSnapshots/Source/LevelSnapshots/Private/Archive/TakeWorldObjectSnapshotArchive.h
@@ -23,10 +23,10 @@ namespace UE::LevelSnapshots::Private
 		virtual bool ShouldSkipProperty(const FProperty* InProperty) const override;
 		//~ End FSnapshotArchive Interface
 
-		protected:
+	protected:
 	
 		//~ Begin FSnapshotArchive Interface
-		virtual UObject* ResolveObjectDependency(int32 ObjectIndex) const { checkNoEntry(); return nullptr; }
+		virtual UObject* ResolveObjectDependency(int32 ObjectIndex, UObject* CurrentValue) const { checkNoEntry(); return nullptr; }
 		//~ End FSnapshotArchive Interface
 
 	private:
diff --git a/Engine/Plugins/VirtualProduction/LevelSnapshots/Source/LevelSnapshots/Private/Data/Util/WorldData/ClassDataUtil.cpp b/Engine/Plugins/VirtualProduction/LevelSnapshots/Source/LevelSnapshots/Private/Data/Util/WorldData/ClassDataUtil.cpp
--- a/Engine/Plugins/VirtualProduction/LevelSnapshots/Source/LevelSnapshots/Private/Data/Util/WorldData/ClassDataUtil.cpp
+++ b/Engine/Plugins/VirtualProduction/LevelSnapshots/Source/LevelSnapshots/Private/Data/Util/WorldData/ClassDataUtil.cpp
@@ -62,7 +62,7 @@ namespace UE::LevelSnapshots::Private
 			return { Cache.ArchetypeObjects[ClassIndex].Get() };
 		}
 		
-		const TOptional<TNonNullPtr<FClassSnapshotData>> ArchetypeData = GetSubobjectArchetypeData(WorldData, ClassIndex, Cache, FallbackInfo);
+		const TOptional<TNonNullPtr<FClassSnapshotData>> ArchetypeData = GetObjectArchetypeData(WorldData, ClassIndex, Cache, FallbackInfo);
 		if (!ArchetypeData)
 		{
 			return {};
@@ -81,7 +81,7 @@ namespace UE::LevelSnapshots::Private
 		return { Archetype };
 	}
 
-	TOptional<TNonNullPtr<FClassSnapshotData>> GetSubobjectArchetypeData(FWorldSnapshotData& WorldData, FClassDataIndex ClassIndex, FSnapshotDataCache& Cache, const FSubobjectArchetypeFallbackInfo& FallbackInfo)
+	TOptional<TNonNullPtr<FClassSnapshotData>> GetObjectArchetypeData(FWorldSnapshotData& WorldData, FClassDataIndex ClassIndex, FSnapshotDataCache& Cache, const FSubobjectArchetypeFallbackInfo& FallbackInfo)
 	{
 		if (!ensureMsgf(WorldData.ClassData.IsValidIndex(ClassIndex), TEXT("This data was supposed to be saved when the snapshot was taken. Data is corrupt.")))
 		{
@@ -121,7 +121,7 @@ namespace UE::LevelSnapshots::Private
 
 	void SerializeClassDefaultsIntoSubobject(UObject* Object, FWorldSnapshotData& WorldData, FClassDataIndex ClassIndex, FSnapshotDataCache& Cache, const FSubobjectArchetypeFallbackInfo& FallbackInfo)
 	{
-		if (const TOptional<TNonNullPtr<FClassSnapshotData>> ArchetypeData = GetSubobjectArchetypeData(WorldData, ClassIndex, Cache, FallbackInfo))
+		if (const TOptional<TNonNullPtr<FClassSnapshotData>> ArchetypeData = GetObjectArchetypeData(WorldData, ClassIndex, Cache, FallbackInfo))
 		{
 			SerializeClassDefaultsIntoSubobject(Object, *ArchetypeData, WorldData);
 		}
@@ -142,6 +142,21 @@ namespace UE::LevelSnapshots::Private
 			FApplyClassDefaulDataArchive::RestoreChangedClassDefaults(DataToSerialize, WorldData, Object);
 		}
 	}
+	
+	void SerializeSelectedClassDefaultsInto(UObject* Object, FWorldSnapshotData& WorldData, FClassDataIndex ClassIndex, FSnapshotDataCache& Cache, const FSubobjectArchetypeFallbackInfo& FallbackInfo, const FPropertySelection& PropertiesToRestore)
+	{
+		if (const TOptional<TNonNullPtr<FClassSnapshotData>> ArchetypeData = GetObjectArchetypeData(WorldData, ClassIndex, Cache, FallbackInfo))
+		{
+			FApplyClassDefaulDataArchive::RestoreSelectedChangedClassDefaults(*ArchetypeData, WorldData, Object, PropertiesToRestore);
+		}
+		else
+		{
+			UE_LOG(LogLevelSnapshots, Warning,
+				TEXT("No CDO saved for class '%s'. If you changed some class default values for this class, then the affected objects will have the latest values instead of the class defaults at the time the snapshot was taken. Should be nothing major to worry about."),
+				*Object->GetClass()->GetName()
+				);
+		}
+	}
 
 	FSoftClassPath GetClass(const FActorSnapshotData& Data, const FWorldSnapshotData& WorldData)
 	{
diff --git a/Engine/Plugins/VirtualProduction/LevelSnapshots/Source/LevelSnapshots/Private/Data/Util/WorldData/ClassDataUtil.h b/Engine/Plugins/VirtualProduction/LevelSnapshots/Source/LevelSnapshots/Private/Data/Util/WorldData/ClassDataUtil.h
--- a/Engine/Plugins/VirtualProduction/LevelSnapshots/Source/LevelSnapshots/Private/Data/Util/WorldData/ClassDataUtil.h
+++ b/Engine/Plugins/VirtualProduction/LevelSnapshots/Source/LevelSnapshots/Private/Data/Util/WorldData/ClassDataUtil.h
@@ -7,6 +7,7 @@
 #include "Templates/NonNullPointer.h"
 #include "UObject/ObjectMacros.h"
 
+struct FPropertySelection;
 class AActor;
 class UObject;
 struct FActorSnapshotData;
@@ -34,11 +35,13 @@ namespace UE::LevelSnapshots::Private
 	TOptional<TNonNullPtr<UObject>> GetSubobjectArchetype(FWorldSnapshotData& WorldData, FClassDataIndex ClassIndex, FSnapshotDataCache& Cache, const FSubobjectArchetypeFallbackInfo& FallbackInfo);
 	
 	/** Gets an archetype that was saved when the snapshot was taken. Each subobject in actors can have their own archetypes. */
-	TOptional<TNonNullPtr<FClassSnapshotData>> GetSubobjectArchetypeData(FWorldSnapshotData& WorldData, FClassDataIndex ClassIndex, FSnapshotDataCache& Cache, const FSubobjectArchetypeFallbackInfo& FallbackInfo);
+	TOptional<TNonNullPtr<FClassSnapshotData>> GetObjectArchetypeData(FWorldSnapshotData& WorldData, FClassDataIndex ClassIndex, FSnapshotDataCache& Cache, const FSubobjectArchetypeFallbackInfo& FallbackInfo);
 	
 	void SerializeClassDefaultsIntoSubobject(UObject* Object, FWorldSnapshotData& WorldData, FClassDataIndex ClassIndex, FSnapshotDataCache& Cache, const FSubobjectArchetypeFallbackInfo& FallbackInfo);
 	void SerializeClassDefaultsIntoSubobject(UObject* Object, FClassSnapshotData& DataToSerialize, FWorldSnapshotData& WorldData);
 
+	void SerializeSelectedClassDefaultsInto(UObject* Object, FWorldSnapshotData& WorldData, FClassDataIndex ClassIndex, FSnapshotDataCache& Cache, const FSubobjectArchetypeFallbackInfo& FallbackInfo, const FPropertySelection& PropertiesToRestore);
+
 	FSoftClassPath GetClass(const FActorSnapshotData& Data, const FWorldSnapshotData& WorldData);
 	FSoftClassPath GetClass(const FSubobjectSnapshotData& Data, const FWorldSnapshotData& WorldData);
 	
diff --git a/Engine/Plugins/VirtualProduction/LevelSnapshots/Source/LevelSnapshots/Private/Data/Util/WorldData/SnapshotObjectUtil.cpp b/Engine/Plugins/VirtualProduction/LevelSnapshots/Source/LevelSnapshots/Private/Data/Util/WorldData/SnapshotObjectUtil.cpp
--- a/Engine/Plugins/VirtualProduction/LevelSnapshots/Source/LevelSnapshots/Private/Data/Util/WorldData/SnapshotObjectUtil.cpp
+++ b/Engine/Plugins/VirtualProduction/LevelSnapshots/Source/LevelSnapshots/Private/Data/Util/WorldData/SnapshotObjectUtil.cpp
@@ -349,8 +349,16 @@ UObject* UE::LevelSnapshots::Private::ResolveObjectDependencyForEditorWorld(FWor
 	return ResolvedObject;
 }
 
-UObject* UE::LevelSnapshots::Private::ResolveObjectDependencyForClassDefaultObject(FWorldSnapshotData& WorldData, int32 ObjectPathIndex)
+UObject* UE::LevelSnapshots::Private::ResolveObjectDependencyForClassDefaultObject(FWorldSnapshotData& WorldData, int32 ObjectPathIndex, UObject* CurrentValue)
 {
+	// If the reference is supposed to be skipped...
+	if (const FObjectReferenceSnapshotMetaData* MetaData = WorldData.SerializedReferenceMetaData.Find(ObjectPathIndex)
+		; MetaData && (MetaData->Flags & EObjectReferenceSnapshotFlags::SkipWhenSerializingArchetypeData) != EObjectReferenceSnapshotFlags::None)
+	{
+		// ... make the property we're serializing point to whatever it was assigned by the class archetype
+		return CurrentValue;
+	}
+	
 	if (ensure(WorldData.SerializedObjectReferences.IsValidIndex(ObjectPathIndex)))
 	{
 		const FSoftObjectPath& OriginalObjectPath = WorldData.SerializedObjectReferences[ObjectPathIndex];
diff --git a/Engine/Plugins/VirtualProduction/LevelSnapshots/Source/LevelSnapshots/Private/Data/Util/WorldData/SnapshotObjectUtil.h b/Engine/Plugins/VirtualProduction/LevelSnapshots/Source/LevelSnapshots/Private/Data/Util/WorldData/SnapshotObjectUtil.h
--- a/Engine/Plugins/VirtualProduction/LevelSnapshots/Source/LevelSnapshots/Private/Data/Util/WorldData/SnapshotObjectUtil.h
+++ b/Engine/Plugins/VirtualProduction/LevelSnapshots/Source/LevelSnapshots/Private/Data/Util/WorldData/SnapshotObjectUtil.h
@@ -15,10 +15,11 @@ struct FWorldSnapshotData;
 
 namespace UE::LevelSnapshots::Private
 {
-	/* Resolves an object dependency for use in the snapshot world. If the object is a subobject, it gets fully serialized. */
+	/** Resolves an object dependency for use in the snapshot world. If the object is a subobject, it gets fully serialized. */
 	UObject* ResolveObjectDependencyForSnapshotWorld(FWorldSnapshotData& WorldData, FSnapshotDataCache& Cache, int32 ObjectPathIndex, const FProcessObjectDependency& ProcessObjectDependency, const FString& LocalisationNamespace); 
 
-	/* Resolves an object dependency for use in the editor world. If the object is a subobject, it is serialized.
+	/**
+	 * Resolves an object dependency for use in the editor world. If the object is a subobject, it is serialized.
 	 * Steps for serializing subobject:
 	 *  - If an equivalent object with the saved name and class exists, use that and serialize the properties in SelectionMap into it.
 	 *  - If object does not exist, allocate a new object and serialize all properties into it
@@ -26,8 +27,13 @@ namespace UE::LevelSnapshots::Private
 	 */
 	UObject* ResolveObjectDependencyForEditorWorld(FWorldSnapshotData& WorldData, FSnapshotDataCache& Cache, int32 ObjectPathIndex, const FString& LocalisationNamespace, const FPropertySelectionMap& SelectionMap);
 
-	/* Resolves an object depedency when restoring a class default object. Simply resolves without further checks. */
-	UObject* ResolveObjectDependencyForClassDefaultObject(FWorldSnapshotData& WorldData, int32 ObjectPathIndex);
+	/**
+	 * Resolves an object depedency when restoring a class default object. Simply resolves without further checks.
+	 * @param WorldData Snapshot data
+	 * @param ObjectPathIndex Identifier the reference in FWorldSnapshotData::SerializedObjectReferences
+	 * @param CurrentValue The value currently saved in the uproperty that we're resolving for
+	 */
+	UObject* ResolveObjectDependencyForClassDefaultObject(FWorldSnapshotData& WorldData, int32 ObjectPathIndex, UObject* CurrentValue);
 
 
 	
diff --git a/Engine/Plugins/VirtualProduction/LevelSnapshots/Source/LevelSnapshots/Private/Data/WorldSnapshotData.cpp b/Engine/Plugins/VirtualProduction/LevelSnapshots/Source/LevelSnapshots/Private/Data/WorldSnapshotData.cpp
--- a/Engine/Plugins/VirtualProduction/LevelSnapshots/Source/LevelSnapshots/Private/Data/WorldSnapshotData.cpp
+++ b/Engine/Plugins/VirtualProduction/LevelSnapshots/Source/LevelSnapshots/Private/Data/WorldSnapshotData.cpp
@@ -48,7 +48,7 @@ namespace UE::LevelSnapshots::Private::Internal
 
 	static void CollectEnums(FArchive& Ar)
 	{
-		TArray<UEnum*> Enums = { StaticEnum<ESnapshotClassFlags>() };
+		TArray<UEnum*> Enums = { StaticEnum<ESnapshotClassFlags>(), StaticEnum<EObjectReferenceSnapshotFlags>() };
 		for (UEnum* Enum : Enums)
 		{
 			FName ClassFlagsName = Enum->GetFName();
@@ -110,6 +110,7 @@ namespace UE::LevelSnapshots::Private::Internal
 		FPropertyIterator(FSubobjectSnapshotData::StaticStruct(), ProcessProperty, ProcessStruct);
 		FPropertyIterator(FCustomSerializationData::StaticStruct(), ProcessProperty, ProcessStruct);
 		FPropertyIterator(FClassSnapshotData::StaticStruct(), ProcessProperty, ProcessStruct);
+		FPropertyIterator(FObjectReferenceSnapshotMetaData::StaticStruct(), ProcessProperty, ProcessStruct);
 	}
 
 	PRAGMA_DISABLE_DEPRECATION_WARNINGS
diff --git a/Engine/Plugins/VirtualProduction/LevelSnapshots/Source/LevelSnapshots/Public/Data/ObjectReferenceSnapshotMetaData.h b/Engine/Plugins/VirtualProduction/LevelSnapshots/Source/LevelSnapshots/Public/Data/ObjectReferenceSnapshotMetaData.h
new file mode 100644
--- /dev/null
+++ b/Engine/Plugins/VirtualProduction/LevelSnapshots/Source/LevelSnapshots/Public/Data/ObjectReferenceSnapshotMetaData.h
@@ -0,0 +1,33 @@
+// Copyright Epic Games, Inc. All Rights Reserved.
+
+#pragma once
+
+#include "CoreMinimal.h"
+#include "ObjectReferenceSnapshotMetaData.generated.h"
+
+UENUM()
+enum class EObjectReferenceSnapshotFlags : uint8
+{
+	None = 0x00000000,
+
+	/**
+	 * When this reference was saved, it was discovered FTakeClassDefaultObjectSnapshotArchive and marked to be skipped.
+	 * 
+	 * The following references are skipped when saving an archetype / CDO:
+	 *	- references to default objects, e.g. if the object RF_ClassDefaultObject | RF_ArchetypeObject flags or the reference points to the serialized archetype.
+	 *	- references to default subobjects, e.g. if the object has the RF_DefaultSubObject flag or is in archetype.
+	 */
+	SkipWhenSerializingArchetypeData = 0x00000001
+};
+
+ENUM_CLASS_FLAGS(EObjectReferenceSnapshotFlags)
+
+/** Holds additional information about a saved object reference (see FWorldSnapshotData::SerializedReferences) */
+USTRUCT()
+struct LEVELSNAPSHOTS_API FObjectReferenceSnapshotMetaData
+{
+	GENERATED_BODY()
+
+	UPROPERTY()
+	EObjectReferenceSnapshotFlags Flags;
+};
\ No newline at end of file
diff --git a/Engine/Plugins/VirtualProduction/LevelSnapshots/Source/LevelSnapshots/Public/Data/SnapshotCustomVersion.h b/Engine/Plugins/VirtualProduction/LevelSnapshots/Source/LevelSnapshots/Public/Data/SnapshotCustomVersion.h
--- a/Engine/Plugins/VirtualProduction/LevelSnapshots/Source/LevelSnapshots/Public/Data/SnapshotCustomVersion.h
+++ b/Engine/Plugins/VirtualProduction/LevelSnapshots/Source/LevelSnapshots/Public/Data/SnapshotCustomVersion.h
@@ -23,12 +23,15 @@ namespace UE::LevelSnapshots
 			/** FWorldSnapshotData::ClassDefaults was replaced by FWorldSnapshotData::ClassData. */
 			ClassArchetypeRefactor = 3,
 
-			/** FWorldSnapshotData now compresses data using oodle before it is saved to disk */
+			/** FWorldSnapshotData now compresses data using Oodle before it is saved to disk */
 			OoddleCompression = 4,
 
 			/** Before this soft object references were saved by FCustomSerializationDataWriter using FObjectAndNameAsStringProxyArchive - if you moved the asset our data would get corrupted. */
 			CustomSubobjectSoftObjectPathRefactor = 5,
 
+			/** FWorldSnapshotData::SerializedReferenceMetaData was introduced. No migration needs to take place. This entry exists only for future proofing (in case something breaks). */
+			AddedSerializedReferenceMetaData = 6,
+
 			// -----<new versions can be added above this line>-------------------------------------------------
 			VersionPlusOne,
 			LatestVersion = VersionPlusOne - 1
diff --git a/Engine/Plugins/VirtualProduction/LevelSnapshots/Source/LevelSnapshots/Public/Data/WorldSnapshotData.h b/Engine/Plugins/VirtualProduction/LevelSnapshots/Source/LevelSnapshots/Public/Data/WorldSnapshotData.h
--- a/Engine/Plugins/VirtualProduction/LevelSnapshots/Source/LevelSnapshots/Public/Data/WorldSnapshotData.h
+++ b/Engine/Plugins/VirtualProduction/LevelSnapshots/Source/LevelSnapshots/Public/Data/WorldSnapshotData.h
@@ -8,6 +8,7 @@
 #include "ClassDefaultObjectSnapshotData.h"
 #include "ClassSnapshotData.h"
 #include "CustomSerializationData.h"
+#include "ObjectReferenceSnapshotMetaData.h"
 #include "SnapshotVersion.h"
 #include "SubobjectSnapshotData.h"
 #include "Templates/NonNullPointer.h"
@@ -62,6 +63,8 @@ struct LEVELSNAPSHOTS_API FWorldSnapshotData
 	 * Actor classes have exactly one entry FClassSnapshotData::FClassSnapshotData.
 	 * For other classes, e.g. components, there may be multiple entries for the same class. This is because
 	 * subobjects can have their own archtetypes (you can see this by looking a the reset to default button in the details panel).
+	 *
+	 * Added in 5.1.
 	 */
 	UPROPERTY()
 	TArray<FClassSnapshotData> ClassData;
@@ -97,20 +100,33 @@ struct LEVELSNAPSHOTS_API FWorldSnapshotData
 	TArray<FSoftObjectPath> SerializedObjectReferences;
 	
 	/**
-	 * Key: A valid index in to SerializedObjectReferences.
+	 * Key: A valid index to SerializedObjectReferences.
 	 * Value: Subobject information for the associated entry in SerializedObjectReferences.
 	 * There is only an entry if the associated object is in fact a subobject. Actors and assets in particular do not get any entry.
+	 *
+	 * Added in 5.0.
 	 */
 	UPROPERTY()
 	TMap<int32, FSubobjectSnapshotData> Subobjects;
 
 	/**
-	 * Key: A valid index in to SerializedObjectReferences
+	 * Key: A valid index to SerializedObjectReferences
 	 * Value: Data that was generated by some ICustomObjectSnapshotSerializer.
+	 * 
+	 * Added in 5.0.
 	 */
 	UPROPERTY()
 	TMap<int32, FCustomSerializationData> CustomSubobjectSerializationData;
 
+	/**
+	 * Key: A valid index to SerializedObjectReferences
+	 * Value: Additional info about the reference, e.g. whether it should be applied by the snapshot.
+	 * 
+	 * Added in 5.1.
+	 */
+	UPROPERTY()
+	TMap<int32, FObjectReferenceSnapshotMetaData> SerializedReferenceMetaData;
+
 	
 
 	/** Binds every entry in SerializedNames to its index. Speeds up adding unique names. */
