Remove unused code (forgot to do last submit)

diff --git a/Engine/Plugins/VirtualProduction/LevelSnapshots/Source/LevelSnapshots/Private/Archive/ApplySnapshotToEditorArchive.cpp b/Engine/Plugins/VirtualProduction/LevelSnapshots/Source/LevelSnapshots/Private/Archive/ApplySnapshotToEditorArchive.cpp
--- a/Engine/Plugins/VirtualProduction/LevelSnapshots/Source/LevelSnapshots/Private/Archive/ApplySnapshotToEditorArchive.cpp
+++ b/Engine/Plugins/VirtualProduction/LevelSnapshots/Source/LevelSnapshots/Private/Archive/ApplySnapshotToEditorArchive.cpp
@@ -5,7 +5,6 @@
 #include "Archive/ClassDefaults/ApplyClassDefaulDataArchive.h"
 #include "Data/RestorationEvents/ApplySnapshotPropertiesScope.h"
 #include "Data/Util/WorldData/ClassDataUtil.h"
-#include "Data/Util/Property/PropertyUtil.h"
 #include "Data/Util/WorldData/SnapshotObjectUtil.h"
 #include "Data/WorldSnapshotData.h"
 #include "LevelSnapshotsLog.h"
@@ -34,10 +33,7 @@ void UE::LevelSnapshots::Private::FApplySnapshotToEditorArchive::ApplyToExisting
 #if WITH_EDITOR
 	InOriginalObject->Modify();
 #endif
-
-	FPropertySelection PropertiesLeftToSerialize = *Selection;
-	auto TrackSerializedProperties = [&PropertiesLeftToSerialize](const FArchiveSerializedPropertyChain* Chain, const FProperty* Property) { PropertiesLeftToSerialize.RemoveProperty(Chain, Property); };
-
+	
 	// 1. Serialize archetype first to handle the case were the archetype has changed properties since the snapshot was taken
 	if (ClassIndex) // Sometimes not available, e.g. for custom subobjects
 	{
@@ -46,7 +42,7 @@ void UE::LevelSnapshots::Private::FApplySnapshotToEditorArchive::ApplyToExisting
 	}
 	
 	// Step 2: Serialise  properties that were different from CDO at time of snapshotting and that are still different from CDO
-	FApplySnapshotToEditorArchive ApplySavedData(InObjectData, InSharedData, InOriginalObject, InSelectionMapForResolvingSubobjects, Selection, TrackSerializedProperties, Cache);
+	FApplySnapshotToEditorArchive ApplySavedData(InObjectData, InSharedData, InOriginalObject, InSelectionMapForResolvingSubobjects, Selection, Cache);
 	InOriginalObject->Serialize(ApplySavedData);
 }
 
@@ -61,9 +57,7 @@ void UE::LevelSnapshots::Private::FApplySnapshotToEditorArchive::ApplyToEditorWo
 	const FApplySnapshotPropertiesScope NotifySnapshotListeners({ InOriginalObject, InSelectionMapForResolvingSubobjects, {}, true });
 	
 	// Apply all properties that we saved into the target actor.
-	// We assume that InOriginalObject was already created with the snapshot CDO as template: we do not need Step 2 from ApplyToExistingWorldObject.
-	auto Dummy = [](auto,auto){};
-	FApplySnapshotToEditorArchive ApplySavedData(InObjectData, InSharedData, InOriginalObject, InSelectionMapForResolvingSubobjects, {}, Dummy, Cache);
+	FApplySnapshotToEditorArchive ApplySavedData(InObjectData, InSharedData, InOriginalObject, InSelectionMapForResolvingSubobjects, {}, Cache);
 	InOriginalObject->Serialize(ApplySavedData);
 }
 
@@ -75,15 +69,12 @@ void UE::LevelSnapshots::Private::FApplySnapshotToEditorArchive::ApplyToEditorWo
 	const FPropertySelectionMap& InSelectionMapForResolvingSubobjects,
 	FClassDataIndex ClassIndex)
 {
-	FPropertySelection PropertiesThatWereSerialized;
-	auto TrackSerializedProperties = [&PropertiesThatWereSerialized](const FArchiveSerializedPropertyChain* Chain, const FProperty* Property) { PropertiesThatWereSerialized.AddProperty({ Chain, Property}); };
-
 	// 1. Serialize archetype first to handle the case were the archetype has changed properties since the snapshot was taken
 	const FSubobjectArchetypeFallbackInfo ClassFallbackInfo{ InOriginalObject->GetOuter(), InOriginalObject->GetFName(), InOriginalObject->GetFlags() };
 	SerializeClassDefaultsIntoSubobject(InOriginalObject, InSharedData, ClassIndex, Cache, ClassFallbackInfo);
 	
 	// Step 2: Apply the data that was different from CDO at time of snapshotting
-	FApplySnapshotToEditorArchive ApplySavedData(InObjectData, InSharedData, InOriginalObject, InSelectionMapForResolvingSubobjects, {}, TrackSerializedProperties, Cache);
+	FApplySnapshotToEditorArchive ApplySavedData(InObjectData, InSharedData, InOriginalObject, InSelectionMapForResolvingSubobjects, {}, Cache);
 	InOriginalObject->Serialize(ApplySavedData);
 }
 
@@ -101,14 +92,6 @@ bool UE::LevelSnapshots::Private::FApplySnapshotToEditorArchive::ShouldSkipPrope
 	return bShouldSkipProperty;
 }
 
-void UE::LevelSnapshots::Private::FApplySnapshotToEditorArchive::PushSerializedProperty(FProperty* InProperty, const bool bIsEditorOnlyProperty)
-{
-	// Do before call to super because super appends InProperty to property chain
-	OnPropertySerialized(GetSerializedPropertyChain(), InProperty);
-	
-	Super::PushSerializedProperty(InProperty, bIsEditorOnlyProperty);
-}
-
 UObject* UE::LevelSnapshots::Private::FApplySnapshotToEditorArchive::ResolveObjectDependency(int32 ObjectIndex, UObject* CurrentValue) const
 {
 	FString LocalizationNamespace;
@@ -124,12 +107,10 @@ UE::LevelSnapshots::Private::FApplySnapshotToEditorArchive::FApplySnapshotToEdit
 	UObject* InOriginalObject,
 	const FPropertySelectionMap& InSelectionMapForResolvingSubobjects,
 	TOptional<const FPropertySelection*> InSelectionSet,
-	FOnSerializeProperty InOnPropertySerialized,
 	FSnapshotDataCache& Cache)
         : Super(InObjectData, InSharedData, true, InOriginalObject)
 		, SelectionMapForResolvingSubobjects(InSelectionMapForResolvingSubobjects)
 		, SelectionSet(InSelectionSet)
-		, OnPropertySerialized(InOnPropertySerialized)
 		, Cache(Cache)
 {
 #if UE_BUILD_DEBUG
diff --git a/Engine/Plugins/VirtualProduction/LevelSnapshots/Source/LevelSnapshots/Private/Archive/ApplySnapshotToEditorArchive.h b/Engine/Plugins/VirtualProduction/LevelSnapshots/Source/LevelSnapshots/Private/Archive/ApplySnapshotToEditorArchive.h
--- a/Engine/Plugins/VirtualProduction/LevelSnapshots/Source/LevelSnapshots/Private/Archive/ApplySnapshotToEditorArchive.h
+++ b/Engine/Plugins/VirtualProduction/LevelSnapshots/Source/LevelSnapshots/Private/Archive/ApplySnapshotToEditorArchive.h
@@ -59,7 +59,6 @@ namespace UE::LevelSnapshots::Private
 	
 		//~ Begin FTakeSnapshotArchiveV2 Interface
 		virtual bool ShouldSkipProperty(const FProperty* InProperty) const override;
-		virtual void PushSerializedProperty(FProperty* InProperty, const bool bIsEditorOnlyProperty) override;
 		//~ End FTakeSnapshotArchiveV2 Interface
 
 	protected:
@@ -70,15 +69,12 @@ namespace UE::LevelSnapshots::Private
 
 	private:
 
-		using FOnSerializeProperty = TFunctionRef<void(const FArchiveSerializedPropertyChain*, const FProperty*)>;
-
 		FApplySnapshotToEditorArchive(
 			FObjectSnapshotData& InObjectData,
 			FWorldSnapshotData& InSharedData,
 			UObject* InOriginalObject,
 			const FPropertySelectionMap& InSelectionMapForResolvingSubobjects,
 			TOptional<const FPropertySelection*> InSelectionSet,
-			FOnSerializeProperty InOnPropertySerialized,
 			FSnapshotDataCache& Cache
 			);
 
@@ -90,9 +86,6 @@ namespace UE::LevelSnapshots::Private
 		/** Immutable list of properties we are supposed to serialise. Serialize all properties if unset. */
 		TOptional<const FPropertySelection*> SelectionSet;
 
-		/** After serialisation is done, we can tell which properties are still left to serialise (i.e. copy from deserialized). */
-		FOnSerializeProperty OnPropertySerialized;
-
 		FSnapshotDataCache& Cache;
 	};
 }
diff --git a/Engine/Plugins/VirtualProduction/LevelSnapshots/Source/LevelSnapshots/Private/Archive/ClassDefaults/BaseClassDefaultArchive.cpp b/Engine/Plugins/VirtualProduction/LevelSnapshots/Source/LevelSnapshots/Private/Archive/ClassDefaults/BaseClassDefaultArchive.cpp
--- a/Engine/Plugins/VirtualProduction/LevelSnapshots/Source/LevelSnapshots/Private/Archive/ClassDefaults/BaseClassDefaultArchive.cpp
+++ b/Engine/Plugins/VirtualProduction/LevelSnapshots/Source/LevelSnapshots/Private/Archive/ClassDefaults/BaseClassDefaultArchive.cpp
@@ -2,7 +2,6 @@
 
 #include "BaseClassDefaultArchive.h"
 
-#include "Data/Util/Property/PropertyUtil.h"
 #include "WorldSnapshotData.h"
 
 UE::LevelSnapshots::Private::FBaseClassDefaultArchive::FBaseClassDefaultArchive(FObjectSnapshotData& InObjectData, FWorldSnapshotData& InSharedData, bool bIsLoading, UObject* InObjectToRestore)
diff --git a/Engine/Plugins/VirtualProduction/LevelSnapshots/Source/LevelSnapshots/Private/Data/Util/Property/PropertyUtil.cpp b/Engine/Plugins/VirtualProduction/LevelSnapshots/Source/LevelSnapshots/Private/Data/Util/Property/PropertyUtil.cpp
deleted file mode 100644
diff --git a/Engine/Plugins/VirtualProduction/LevelSnapshots/Source/LevelSnapshots/Private/Data/Util/Property/PropertyUtil.h b/Engine/Plugins/VirtualProduction/LevelSnapshots/Source/LevelSnapshots/Private/Data/Util/Property/PropertyUtil.h
deleted file mode 100644
diff --git a/Engine/Plugins/VirtualProduction/LevelSnapshots/Source/LevelSnapshots/Private/Data/Util/Property/WorldReferenceCheckingUtil.cpp b/Engine/Plugins/VirtualProduction/LevelSnapshots/Source/LevelSnapshots/Private/Data/Util/Property/WorldReferenceCheckingUtil.cpp
deleted file mode 100644
diff --git a/Engine/Plugins/VirtualProduction/LevelSnapshots/Source/LevelSnapshots/Private/Data/Util/Property/WorldReferenceCheckingUtil.h b/Engine/Plugins/VirtualProduction/LevelSnapshots/Source/LevelSnapshots/Private/Data/Util/Property/WorldReferenceCheckingUtil.h
deleted file mode 100644
