Fix crash due to uninitialised object references

diff --git a/Engine/Plugins/VirtualProduction/LevelSnapshots/Source/LevelSnapshots/Private/Archive/ClassDefaults/BaseClassDefaultArchive.h b/Engine/Plugins/VirtualProduction/LevelSnapshots/Source/LevelSnapshots/Private/Archive/ClassDefaults/BaseClassDefaultArchive.h
--- a/Engine/Plugins/VirtualProduction/LevelSnapshots/Source/LevelSnapshots/Private/Archive/ClassDefaults/BaseClassDefaultArchive.h
+++ b/Engine/Plugins/VirtualProduction/LevelSnapshots/Source/LevelSnapshots/Private/Archive/ClassDefaults/BaseClassDefaultArchive.h
@@ -18,5 +18,8 @@ namespace UE::LevelSnapshots::Private
 	protected:
 	
 		FBaseClassDefaultArchive(FObjectSnapshotData& InObjectData, FWorldSnapshotData& InSharedData, bool bIsLoading, UObject* InObjectToRestore);
+
+		// Class archetypes do not need fixing up... they only reference assets and subobjects. Subobjects are handled entirely by the the archetype itself.
+		virtual void FixUpReference(UObject*& Value) const override {}
 	};
 }
\ No newline at end of file
diff --git a/Engine/Plugins/VirtualProduction/LevelSnapshots/Source/LevelSnapshots/Private/Archive/ClassDefaults/TakeClassDefaultObjectSnapshotArchive.cpp b/Engine/Plugins/VirtualProduction/LevelSnapshots/Source/LevelSnapshots/Private/Archive/ClassDefaults/TakeClassDefaultObjectSnapshotArchive.cpp
--- a/Engine/Plugins/VirtualProduction/LevelSnapshots/Source/LevelSnapshots/Private/Archive/ClassDefaults/TakeClassDefaultObjectSnapshotArchive.cpp
+++ b/Engine/Plugins/VirtualProduction/LevelSnapshots/Source/LevelSnapshots/Private/Archive/ClassDefaults/TakeClassDefaultObjectSnapshotArchive.cpp
@@ -6,36 +6,79 @@
 #include "ObjectSnapshotData.h"
 #include "WorldSnapshotData.h"
 
+namespace UE::LevelSnapshots::Private
+{
+	class FCollectSubobjectPropertyArchive : public FArchiveUObject
+	{
+		using Super = FArchiveUObject;
+		const UObject* SerializedObject;
+	public:
+
+		TSet<const FProperty*> PropertiesContainingSubobjects;
+
+		FCollectSubobjectPropertyArchive(const UObject* SerializedObject)
+			: SerializedObject(SerializedObject)
+		{
+			Super::SetIsSaving(true);
+			Super::SetIsPersistent(true);
+			// Intentionally not setting ArIsObjectReferenceCollector because we want to find them "naturally" to determine referencing properties
+		}
+		
+		virtual FArchive& operator<<(UObject*& Object) override
+		{
+			if (IsSkippedObject(Object) && GetSerializedProperty())
+			{
+				PropertiesContainingSubobjects.Add(GetSerializedProperty());
+			}
+			return *this;
+		}
+
+		bool IsSkippedObject(const UObject* Object)
+		{
+			if (Object)
+			{
+				const bool bIsClassDefault = Object->HasAnyFlags(RF_ClassDefaultObject | RF_ArchetypeObject);
+				const bool bIsPointingToDefaultSubobject = Object->HasAnyFlags(RF_DefaultSubObject);
+				const bool bIsPointingToSelf = Object == SerializedObject;
+				const bool bIsPointingToSubobject = Object->IsIn(SerializedObject);
+				const bool bShouldSkip = bIsClassDefault || bIsPointingToDefaultSubobject || bIsPointingToSelf || bIsPointingToSubobject;
+
+				// We don't ever want to interfere with the default subobjects that the class archetype assigns so the reference is marked so we know it is supposed to be skipped when restoring 
+				return bShouldSkip;
+			}
+			return false;
+		}
+	};
+}
+
 void UE::LevelSnapshots::Private::FTakeClassDefaultObjectSnapshotArchive::SaveClassDefaultObject(FClassSnapshotData& InObjectData, FWorldSnapshotData& InSharedData, UObject* SerializedObject)
 {
-	FTakeClassDefaultObjectSnapshotArchive SaveClass(InObjectData, InSharedData, SerializedObject);
+	// We don't want to interfere with subobjects created by archetypes / CDOs because restoring those is very involved so we skip properties that contain them
+	FCollectSubobjectPropertyArchive CollectSubobjectProperties(SerializedObject);
+	SerializedObject->Serialize(CollectSubobjectProperties);
+	
+	FTakeClassDefaultObjectSnapshotArchive SaveClass(InObjectData, InSharedData, SerializedObject, MoveTemp(CollectSubobjectProperties.PropertiesContainingSubobjects));
 	SerializedObject->Serialize(SaveClass);
 	InObjectData.ObjectFlags = SerializedObject->GetFlags();
 	InObjectData.ClassPath = SerializedObject->GetClass();
 	InObjectData.ClassFlags = SerializedObject->GetClass()->GetFlags();
 }
 
-void UE::LevelSnapshots::Private::FTakeClassDefaultObjectSnapshotArchive::OnAddObjectDependency(int32 ObjectIndex, UObject* Object) const
+bool UE::LevelSnapshots::Private::FTakeClassDefaultObjectSnapshotArchive::ShouldSkipProperty(const FProperty* InProperty) const
 {
-	if (Object)
-	{
-		const bool bIsClassDefault = Object->HasAnyFlags(RF_ClassDefaultObject | RF_ArchetypeObject);
-		const bool bIsPointingToDefaultSubobject = Object->HasAnyFlags(RF_DefaultSubObject);
-		const bool bIsPointingToSelf = Object == GetSerializedObject();
-		const bool bIsPointingToSubobject = Object->IsIn(GetSerializedObject());
-		const bool bShouldSkip = bIsClassDefault || bIsPointingToDefaultSubobject || bIsPointingToSelf || bIsPointingToSubobject;
-
-		// We don't ever want to interfere with the default subobjects that the class archetype assigns so the reference is marked so we know it is supposed to be skipped when restoring 
-		if (bShouldSkip)
-		{
-			GetSharedData().SerializedReferenceMetaData.FindOrAdd(ObjectIndex).Flags |= EObjectReferenceSnapshotFlags::SkipWhenSerializingArchetypeData;
-		}
-	}
+	const bool bIsSkippedProperty = SkippedProperties.Contains(InProperty);
+	UE_CLOG(bIsSkippedProperty, LogLevelSnapshots, Verbose, TEXT("Skipping property %s of object %s because it references a subobject."), *InProperty->GetName(), *GetSerializedObject()->GetPathName());
+	return bIsSkippedProperty || Super::ShouldSkipProperty(InProperty);
 }
 
-UE::LevelSnapshots::Private::FTakeClassDefaultObjectSnapshotArchive::FTakeClassDefaultObjectSnapshotArchive(FObjectSnapshotData& InObjectData, FWorldSnapshotData& InSharedData, UObject* InSerializedObject)
-	:
-	Super(InObjectData, InSharedData, false, InSerializedObject)
+UE::LevelSnapshots::Private::FTakeClassDefaultObjectSnapshotArchive::FTakeClassDefaultObjectSnapshotArchive(
+	FObjectSnapshotData& InObjectData,
+	FWorldSnapshotData& InSharedData,
+	UObject* InSerializedObject,
+	TSet<const FProperty*> PropertiesToSkip
+	)
+	: Super(InObjectData, InSharedData, false, InSerializedObject)
+	, SkippedProperties(MoveTemp(PropertiesToSkip))
 {
 #if UE_BUILD_DEBUG
 	UE_LOG(LogLevelSnapshots, VeryVerbose, TEXT("FTakeClassDefaultObjectSnapshotArchive: %s (%s)"), *InSerializedObject->GetPathName(), *InSerializedObject->GetClass()->GetPathName());
diff --git a/Engine/Plugins/VirtualProduction/LevelSnapshots/Source/LevelSnapshots/Private/Archive/ClassDefaults/TakeClassDefaultObjectSnapshotArchive.h b/Engine/Plugins/VirtualProduction/LevelSnapshots/Source/LevelSnapshots/Private/Archive/ClassDefaults/TakeClassDefaultObjectSnapshotArchive.h
--- a/Engine/Plugins/VirtualProduction/LevelSnapshots/Source/LevelSnapshots/Private/Archive/ClassDefaults/TakeClassDefaultObjectSnapshotArchive.h
+++ b/Engine/Plugins/VirtualProduction/LevelSnapshots/Source/LevelSnapshots/Private/Archive/ClassDefaults/TakeClassDefaultObjectSnapshotArchive.h
@@ -20,16 +20,22 @@ namespace UE::LevelSnapshots::Private
 
 		static void SaveClassDefaultObject(FClassSnapshotData& InObjectData, FWorldSnapshotData& InSharedData, UObject* SerializedObject);
 
+		//~ Begin FArchive Interface
+		virtual bool ShouldSkipProperty(const FProperty* InProperty) const override;
+		//~ End FArchive Interface
+
 	protected:
 
 		//~ Begin FSnapshotArchive Interface
 		virtual UObject* ResolveObjectDependency(int32 ObjectIndex, UObject* CurrentValue) const override { checkNoEntry(); return nullptr; }
-		virtual void OnAddObjectDependency(int32 ObjectIndex, UObject* Object) const override;
+		virtual void OnAddObjectDependency(int32 ObjectIndex, UObject* Object) const override {}
 		//~ End FSnapshotArchive Interface
 		
 	private:
 	
-		FTakeClassDefaultObjectSnapshotArchive(FObjectSnapshotData& InObjectData, FWorldSnapshotData& InSharedData, UObject* InSerializedObject);
+		FTakeClassDefaultObjectSnapshotArchive(FObjectSnapshotData& InObjectData, FWorldSnapshotData& InSharedData, UObject* InSerializedObject, TSet<const FProperty*> PropertiesToSkip);
+
+		const TSet<const FProperty*> SkippedProperties;
 	};
 }
 
diff --git a/Engine/Plugins/VirtualProduction/LevelSnapshots/Source/LevelSnapshots/Private/Archive/SnapshotArchive.cpp b/Engine/Plugins/VirtualProduction/LevelSnapshots/Source/LevelSnapshots/Private/Archive/SnapshotArchive.cpp
--- a/Engine/Plugins/VirtualProduction/LevelSnapshots/Source/LevelSnapshots/Private/Archive/SnapshotArchive.cpp
+++ b/Engine/Plugins/VirtualProduction/LevelSnapshots/Source/LevelSnapshots/Private/Archive/SnapshotArchive.cpp
@@ -124,19 +124,7 @@ FArchive& UE::LevelSnapshots::Private::FSnapshotArchive::operator<<(UObject*& Va
 		}
 
 		Value = ResolveObjectDependency(ReferencedIndex, Value);
-		if (Value)
-		{
-			const ULevel* ResolvedOwningLevel = Value->GetTypedOuter<ULevel>();
-			const ULevel* SerializedObjectOwningLevel = SerializedObject->GetTypedOuter<ULevel>();
-			if (ResolvedOwningLevel && ResolvedOwningLevel != SerializedObjectOwningLevel)
-			{
-				// Example how this could happen:
-				// 1. Make actor a reference actor b
-				// 2. Move actor b from level Foo to level Bar (select actor > right-click Bar in Level tab > Move selected actors to level).
-				UE_LOG(LogLevelSnapshots, Warning, TEXT("Resolved world object %s for %s. Referencing objects from other world / levels is not permitted: Nulling reference."), *Value->GetPathName(), *SerializedObject->GetPathName());
-				Value = nullptr;
-			}
-		}
+		FixUpReference(Value);
 	}
 	else
 	{
@@ -179,6 +167,23 @@ void UE::LevelSnapshots::Private::FSnapshotArchive::Serialize(void* Data, int64
 	}
 }
 
+void UE::LevelSnapshots::Private::FSnapshotArchive::FixUpReference(UObject*& Value) const
+{
+	if (IsValid(Value))
+	{
+		const ULevel* ResolvedOwningLevel = Value->GetTypedOuter<ULevel>();
+		const ULevel* SerializedObjectOwningLevel = SerializedObject->GetTypedOuter<ULevel>();
+		if (ResolvedOwningLevel && ResolvedOwningLevel != SerializedObjectOwningLevel)
+		{
+			// Example how this could happen:
+			// 1. Make actor a reference actor b
+			// 2. Move actor b from level Foo to level Bar (select actor > right-click Bar in Level tab > Move selected actors to level).
+			UE_LOG(LogLevelSnapshots, Warning, TEXT("Resolved world object %s for %s. Referencing objects from other world / levels is not permitted: Nulling reference."), *Value->GetPathName(), *SerializedObject->GetPathName());
+			Value = nullptr;
+		}
+	}
+}
+
 UE::LevelSnapshots::Private::FSnapshotArchive::FSnapshotArchive(FObjectSnapshotData& InObjectData, FWorldSnapshotData& InSharedData, bool bIsLoading, UObject* InSerializedObject)
 	:
 	ExcludedPropertyFlags(CPF_BlueprintAssignable | CPF_Transient | CPF_Deprecated),
diff --git a/Engine/Plugins/VirtualProduction/LevelSnapshots/Source/LevelSnapshots/Private/Archive/SnapshotArchive.h b/Engine/Plugins/VirtualProduction/LevelSnapshots/Source/LevelSnapshots/Private/Archive/SnapshotArchive.h
--- a/Engine/Plugins/VirtualProduction/LevelSnapshots/Source/LevelSnapshots/Private/Archive/SnapshotArchive.h
+++ b/Engine/Plugins/VirtualProduction/LevelSnapshots/Source/LevelSnapshots/Private/Archive/SnapshotArchive.h
@@ -37,6 +37,8 @@ namespace UE::LevelSnapshots::Private
 		virtual UObject* ResolveObjectDependency(int32 ObjectIndex, UObject* CurrentValue) const = 0;
 		/** Called when an object dependency is saved. */
 		virtual void OnAddObjectDependency(int32 ObjectIndex, UObject* Object) const {}
+
+		virtual void FixUpReference(UObject*& Value) const;
 	
 		FWorldSnapshotData& GetSharedData() const { return SharedData;}
 		UObject* GetSerializedObject() const { return SerializedObject; }
diff --git a/Engine/Plugins/VirtualProduction/LevelSnapshots/Source/LevelSnapshots/Private/Data/LevelSnapshot.cpp b/Engine/Plugins/VirtualProduction/LevelSnapshots/Source/LevelSnapshots/Private/Data/LevelSnapshot.cpp
--- a/Engine/Plugins/VirtualProduction/LevelSnapshots/Source/LevelSnapshots/Private/Data/LevelSnapshot.cpp
+++ b/Engine/Plugins/VirtualProduction/LevelSnapshots/Source/LevelSnapshots/Private/Data/LevelSnapshot.cpp
@@ -471,44 +471,6 @@ void ULevelSnapshot::RecreateSnapshotWorld()
 	EnsureWorldInitialised();
 }
 
-namespace UE::LevelSnapshots::Private
-{
-	UWorld* FindOrCreateWorld(const TArray<UWorld*>& AllWorlds, UObject* Outer, FName Name);
-}
-
-void ULevelSnapshot::ReinitializeWorld()
-{
-	using namespace UE::LevelSnapshots::Private;
-	
-	TArray<UWorld*> AllWorlds = MoveTemp(SnapshotSublevels);
-	AllWorlds.Add(RootSnapshotWorld);
-
-	const FName RootWorldName = *ExtractPathWithoutSubobjects(MapPath).GetAssetName();
-	RootSnapshotWorld = FindOrCreateWorld(AllWorlds, this, RootWorldName);
-		
-	TSet<FName> Sublevels = ExtractLevelNames(SerializedData);
-	Sublevels.Remove(RootWorldName);
-	SnapshotSublevels.Reserve(Sublevels.Num());
-	for (const FName& SublevelName : Sublevels)
-	{
-		SnapshotSublevels.Add(FindOrCreateWorld(AllWorlds, this, SublevelName));
-	}
-}
-
-namespace UE::LevelSnapshots::Private
-{
-	UWorld* FindOrCreateWorld(const TArray<UWorld*>& AllWorlds, UObject* Outer, FName Name)
-	{
-		const int32 Index = Algo::IndexOfByPredicate(AllWorlds, [Name](UWorld* World){ return World->GetFName() == Name; });
-		if (AllWorlds.IsValidIndex(Index))
-		{
-			InitializeWorld(AllWorlds[Index]);
-			return AllWorlds[Index];
-		}
-		return CreateWorld(Outer, Name);
-	}
-}
-
 #if WITH_EDITOR
 void ULevelSnapshot::ResetDiffCacheToUninitialized(TArrayView<AActor*> ModifiedActors)
 {
diff --git a/Engine/Plugins/VirtualProduction/LevelSnapshots/Source/LevelSnapshots/Private/Data/Util/WorldData/SnapshotObjectUtil.cpp b/Engine/Plugins/VirtualProduction/LevelSnapshots/Source/LevelSnapshots/Private/Data/Util/WorldData/SnapshotObjectUtil.cpp
--- a/Engine/Plugins/VirtualProduction/LevelSnapshots/Source/LevelSnapshots/Private/Data/Util/WorldData/SnapshotObjectUtil.cpp
+++ b/Engine/Plugins/VirtualProduction/LevelSnapshots/Source/LevelSnapshots/Private/Data/Util/WorldData/SnapshotObjectUtil.cpp
@@ -351,14 +351,6 @@ UObject* UE::LevelSnapshots::Private::ResolveObjectDependencyForEditorWorld(FWor
 
 UObject* UE::LevelSnapshots::Private::ResolveObjectDependencyForClassDefaultObject(FWorldSnapshotData& WorldData, int32 ObjectPathIndex, UObject* CurrentValue)
 {
-	// If the reference is supposed to be skipped...
-	if (const FObjectReferenceSnapshotMetaData* MetaData = WorldData.SerializedReferenceMetaData.Find(ObjectPathIndex)
-		; MetaData && (MetaData->Flags & EObjectReferenceSnapshotFlags::SkipWhenSerializingArchetypeData) != EObjectReferenceSnapshotFlags::None)
-	{
-		// ... make the property we're serializing point to whatever it was assigned by the class archetype
-		return CurrentValue;
-	}
-	
 	if (ensure(WorldData.SerializedObjectReferences.IsValidIndex(ObjectPathIndex)))
 	{
 		const FSoftObjectPath& OriginalObjectPath = WorldData.SerializedObjectReferences[ObjectPathIndex];
diff --git a/Engine/Plugins/VirtualProduction/LevelSnapshots/Source/LevelSnapshots/Private/Data/WorldSnapshotData.cpp b/Engine/Plugins/VirtualProduction/LevelSnapshots/Source/LevelSnapshots/Private/Data/WorldSnapshotData.cpp
--- a/Engine/Plugins/VirtualProduction/LevelSnapshots/Source/LevelSnapshots/Private/Data/WorldSnapshotData.cpp
+++ b/Engine/Plugins/VirtualProduction/LevelSnapshots/Source/LevelSnapshots/Private/Data/WorldSnapshotData.cpp
@@ -48,7 +48,7 @@ namespace UE::LevelSnapshots::Private::Internal
 
 	static void CollectEnums(FArchive& Ar)
 	{
-		TArray<UEnum*> Enums = { StaticEnum<ESnapshotClassFlags>(), StaticEnum<EObjectReferenceSnapshotFlags>() };
+		TArray<UEnum*> Enums = { StaticEnum<ESnapshotClassFlags>() };
 		for (UEnum* Enum : Enums)
 		{
 			FName ClassFlagsName = Enum->GetFName();
@@ -110,7 +110,6 @@ namespace UE::LevelSnapshots::Private::Internal
 		FPropertyIterator(FSubobjectSnapshotData::StaticStruct(), ProcessProperty, ProcessStruct);
 		FPropertyIterator(FCustomSerializationData::StaticStruct(), ProcessProperty, ProcessStruct);
 		FPropertyIterator(FClassSnapshotData::StaticStruct(), ProcessProperty, ProcessStruct);
-		FPropertyIterator(FObjectReferenceSnapshotMetaData::StaticStruct(), ProcessProperty, ProcessStruct);
 	}
 
 	PRAGMA_DISABLE_DEPRECATION_WARNINGS
diff --git a/Engine/Plugins/VirtualProduction/LevelSnapshots/Source/LevelSnapshots/Public/Data/LevelSnapshot.h b/Engine/Plugins/VirtualProduction/LevelSnapshots/Source/LevelSnapshots/Public/Data/LevelSnapshot.h
--- a/Engine/Plugins/VirtualProduction/LevelSnapshots/Source/LevelSnapshots/Public/Data/LevelSnapshot.h
+++ b/Engine/Plugins/VirtualProduction/LevelSnapshots/Source/LevelSnapshots/Public/Data/LevelSnapshot.h
@@ -123,7 +123,6 @@ private:
 	void ClearCache();
 	
 	void RecreateSnapshotWorld();
-	void ReinitializeWorld();
 	
 	/** Callback to destroy our world when editor (editor build) or play (game builds) world is destroyed. */
 	FDelegateHandle Handle;
diff --git a/Engine/Plugins/VirtualProduction/LevelSnapshots/Source/LevelSnapshots/Public/Data/ObjectReferenceSnapshotMetaData.h b/Engine/Plugins/VirtualProduction/LevelSnapshots/Source/LevelSnapshots/Public/Data/ObjectReferenceSnapshotMetaData.h
deleted file mode 100644
diff --git a/Engine/Plugins/VirtualProduction/LevelSnapshots/Source/LevelSnapshots/Public/Data/SnapshotCustomVersion.h b/Engine/Plugins/VirtualProduction/LevelSnapshots/Source/LevelSnapshots/Public/Data/SnapshotCustomVersion.h
--- a/Engine/Plugins/VirtualProduction/LevelSnapshots/Source/LevelSnapshots/Public/Data/SnapshotCustomVersion.h
+++ b/Engine/Plugins/VirtualProduction/LevelSnapshots/Source/LevelSnapshots/Public/Data/SnapshotCustomVersion.h
@@ -29,8 +29,11 @@ namespace UE::LevelSnapshots
 			/** Before this soft object references were saved by FCustomSerializationDataWriter using FObjectAndNameAsStringProxyArchive - if you moved the asset our data would get corrupted. */
 			CustomSubobjectSoftObjectPathRefactor = 5,
 
-			/** FWorldSnapshotData::SerializedReferenceMetaData was introduced. No migration needs to take place. This entry exists only for future proofing (in case something breaks). */
-			AddedSerializedReferenceMetaData = 6,
+			/**
+			 * When serializing a class archetype, we now skip properties that contain subobjects by discovering reference properties with a reference collector archive.
+			 * No migration needs to take place. This entry exists only for future proofing (in case something breaks).
+			 */
+			SkipSubobjectReferencesWhenSerializingArchetypesRefactor = 6,
 
 			// -----<new versions can be added above this line>-------------------------------------------------
 			VersionPlusOne,
diff --git a/Engine/Plugins/VirtualProduction/LevelSnapshots/Source/LevelSnapshots/Public/Data/WorldSnapshotData.h b/Engine/Plugins/VirtualProduction/LevelSnapshots/Source/LevelSnapshots/Public/Data/WorldSnapshotData.h
--- a/Engine/Plugins/VirtualProduction/LevelSnapshots/Source/LevelSnapshots/Public/Data/WorldSnapshotData.h
+++ b/Engine/Plugins/VirtualProduction/LevelSnapshots/Source/LevelSnapshots/Public/Data/WorldSnapshotData.h
@@ -8,7 +8,6 @@
 #include "ClassDefaultObjectSnapshotData.h"
 #include "ClassSnapshotData.h"
 #include "CustomSerializationData.h"
-#include "ObjectReferenceSnapshotMetaData.h"
 #include "SnapshotVersion.h"
 #include "SubobjectSnapshotData.h"
 #include "Templates/NonNullPointer.h"
@@ -118,15 +117,6 @@ struct LEVELSNAPSHOTS_API FWorldSnapshotData
 	UPROPERTY()
 	TMap<int32, FCustomSerializationData> CustomSubobjectSerializationData;
 
-	/**
-	 * Key: A valid index to SerializedObjectReferences
-	 * Value: Additional info about the reference, e.g. whether it should be applied by the snapshot.
-	 * 
-	 * Added in 5.1.
-	 */
-	UPROPERTY()
-	TMap<int32, FObjectReferenceSnapshotMetaData> SerializedReferenceMetaData;
-
 	
 
 	/** Binds every entry in SerializedNames to its index. Speeds up adding unique names. */
