Level instances now have their transform restored

diff --git a/Engine/Plugins/VirtualProduction/LevelSnapshots/Source/LevelSnapshots/Private/Restorability/EngineTypesRestorationFence.cpp b/Engine/Plugins/VirtualProduction/LevelSnapshots/Source/LevelSnapshots/Private/Restorability/EngineTypesRestorationFence.cpp
--- a/Engine/Plugins/VirtualProduction/LevelSnapshots/Source/LevelSnapshots/Private/Restorability/EngineTypesRestorationFence.cpp
+++ b/Engine/Plugins/VirtualProduction/LevelSnapshots/Source/LevelSnapshots/Private/Restorability/EngineTypesRestorationFence.cpp
@@ -4,6 +4,7 @@
 
 #include "ActorGroupRestoration.h"
 #include "AttachParentShowsTransformProperties.h"
+#include "LevelInstanceRestoration.h"
 #include "Restorability/CollisionRestoration.h"
 #include "Restorability/GridPlacementRestoration.h"
 #include "SpecialActorPropertySupport.h"
@@ -17,7 +18,7 @@
 
 namespace UE::LevelSnapshots::Private::EngineTypesRestorationFence
 {
-	static void AddSoftObjectPathSupport(UE::LevelSnapshots::Private::FLevelSnapshotsModule& Module)
+	static void AddSoftObjectPathSupport(ILevelSnapshotsModule& Module)
 	{
 		// FSnapshotRestorability::IsRestorableProperty requires properties to have the CPF_Edit specifier
 		// FSoftObjectPath does not have this so we need to explicitly allow its properties
@@ -33,7 +34,7 @@ namespace UE::LevelSnapshots::Private::EngineTypesRestorationFence
 		Module.AddExplicitilySupportedProperties(SoftObjectPathProperties);
 	}
 
-	static void AddAttachParentSupport(UE::LevelSnapshots::Private::FLevelSnapshotsModule& Module)
+	static void AddAttachParentSupport(ILevelSnapshotsModule& Module)
 	{
 		// These properties are not visible by default because they're not CPF_Edit
 		const FProperty* AttachParent = USceneComponent::StaticClass()->FindPropertyByName(FName("AttachParent"));
@@ -46,7 +47,7 @@ namespace UE::LevelSnapshots::Private::EngineTypesRestorationFence
 		}
 	}
 
-	static void DisableIrrelevantBrushSubobjects(UE::LevelSnapshots::Private::FLevelSnapshotsModule& Module)
+	static void DisableIrrelevantBrushSubobjects(ILevelSnapshotsModule& Module)
 	{
 #if WITH_EDITORONLY_DATA
 		// ABrush::BrushBuilder is CPF_Edit but no user ever cares about it. We don't want it to make volumes to show up as changed.
@@ -58,7 +59,7 @@ namespace UE::LevelSnapshots::Private::EngineTypesRestorationFence
 #endif
 	}
 
-	static void DisableIrrelevantWorldSettings(UE::LevelSnapshots::Private::FLevelSnapshotsModule& Module)
+	static void DisableIrrelevantWorldSettings(ILevelSnapshotsModule& Module)
 	{
 		// AWorldSettings::NavigationSystemConfig is CPF_Edit but no user ever cares about it.
 		const FProperty* NavigationSystemConfig = AWorldSettings::StaticClass()->FindPropertyByName(FName("NavigationSystemConfig"));
@@ -68,7 +69,7 @@ namespace UE::LevelSnapshots::Private::EngineTypesRestorationFence
 		}
 	}
 
-	static void DisableIrrelevantMaterialInstanceProperties(UE::LevelSnapshots::Private::FLevelSnapshotsModule& Module)
+	static void DisableIrrelevantMaterialInstanceProperties(ILevelSnapshotsModule& Module)
 	{
 		// This property causes diffs sometimes for unexplained reasons when creating in construction script... does not seem to be important
 		const FProperty* BasePropertyOverrides = UMaterialInstance::StaticClass()->FindPropertyByName(GET_MEMBER_NAME_CHECKED(UMaterialInstance, BasePropertyOverrides));
@@ -78,7 +79,7 @@ namespace UE::LevelSnapshots::Private::EngineTypesRestorationFence
 		}
 	}
 
-	static void DisableIrrelevantActorProperties(UE::LevelSnapshots::Private::FLevelSnapshotsModule& Module)
+	static void DisableIrrelevantActorProperties(ILevelSnapshotsModule& Module)
 	{
 #if WITH_EDITORONLY_DATA
 		const FProperty* ActorGuid = AActor::StaticClass()->FindPropertyByName(FName("ActorGuid"));
@@ -97,7 +98,7 @@ namespace UE::LevelSnapshots::Private::EngineTypesRestorationFence
 		}
 	}
 
-	static void DisableClothingSimulationFactory(UE::LevelSnapshots::Private::FLevelSnapshotsModule& Module)
+	static void DisableClothingSimulationFactory(ILevelSnapshotsModule& Module)
 	{
 		// This property is directly set by RegisterComponent and similar functions... we have no chance of restoring this property properly
 		const FProperty* ClothingSimulationFactory = USkeletalMeshComponent::StaticClass()->FindPropertyByName(GET_MEMBER_NAME_CHECKED(USkeletalMeshComponent, ClothingSimulationFactory));
@@ -124,5 +125,6 @@ namespace UE::LevelSnapshots::Private::EngineTypesRestorationFence
 		ActorGroupRestoration::Register(Module);
 		SpecialActorPropertySupport::Register(Module);
 		AttachParentShowsTransformPropertiesFix::Register(Module);
+		LevelInstanceRestoration::Register(Module);
 	}
 }
diff --git a/Engine/Plugins/VirtualProduction/LevelSnapshots/Source/LevelSnapshots/Private/Restorability/LevelInstanceRestoration.cpp b/Engine/Plugins/VirtualProduction/LevelSnapshots/Source/LevelSnapshots/Private/Restorability/LevelInstanceRestoration.cpp
new file mode 100644
--- /dev/null
+++ b/Engine/Plugins/VirtualProduction/LevelSnapshots/Source/LevelSnapshots/Private/Restorability/LevelInstanceRestoration.cpp
@@ -0,0 +1,32 @@
+// Copyright Epic Games, Inc. All Rights Reserved.
+
+#include "LevelInstanceRestoration.h"
+
+#include "ILevelSnapshotsModule.h"
+#include "Restorability/Interfaces/ISnapshotRestorabilityOverrider.h"
+
+#include "LevelInstance/LevelInstanceActor.h"
+
+namespace UE::LevelSnapshots::Private::LevelInstanceRestoration
+{
+	/** Adds support for restoring level instances */
+	class FLevelInstanceRestoration : public ISnapshotRestorabilityOverrider
+	{
+	public:
+		
+		virtual ERestorabilityOverride IsComponentDesirableForCapture(const UActorComponent* Component) override
+		{
+			// The root component is native and does not have a corresponding UPROPERTY() thus FComponentEditorUtils::CanEditComponentInstance will return false.
+			// That would skip the capture of the component - we explicitly allow it (and non-root components for no particular reason) as well.
+			return Component->GetOwner()->IsA<ALevelInstance>()
+				? ERestorabilityOverride::Allow
+				: ERestorabilityOverride::DoNotCare;
+		}
+	};
+	
+	void Register(ILevelSnapshotsModule& Module)
+	{
+		const TSharedRef<FLevelInstanceRestoration> LevelInstanceSupport = MakeShared<FLevelInstanceRestoration>();
+		Module.RegisterRestorabilityOverrider(LevelInstanceSupport);
+	}
+}
diff --git a/Engine/Plugins/VirtualProduction/LevelSnapshots/Source/LevelSnapshots/Private/Restorability/LevelInstanceRestoration.h b/Engine/Plugins/VirtualProduction/LevelSnapshots/Source/LevelSnapshots/Private/Restorability/LevelInstanceRestoration.h
new file mode 100644
--- /dev/null
+++ b/Engine/Plugins/VirtualProduction/LevelSnapshots/Source/LevelSnapshots/Private/Restorability/LevelInstanceRestoration.h
@@ -0,0 +1,17 @@
+// Copyright Epic Games, Inc. All Rights Reserved.
+
+#pragma once
+
+#include "CoreMinimal.h"
+
+namespace UE::LevelSnapshots
+{
+	class ILevelSnapshotsModule;
+}
+
+namespace UE::LevelSnapshots::Private::LevelInstanceRestoration
+{
+	void Register(ILevelSnapshotsModule& Module);
+}
+
+
diff --git a/Engine/Plugins/VirtualProduction/LevelSnapshots/Source/LevelSnapshots/Private/Restorability/SnapshotRestorability.cpp b/Engine/Plugins/VirtualProduction/LevelSnapshots/Source/LevelSnapshots/Private/Restorability/SnapshotRestorability.cpp
--- a/Engine/Plugins/VirtualProduction/LevelSnapshots/Source/LevelSnapshots/Private/Restorability/SnapshotRestorability.cpp
+++ b/Engine/Plugins/VirtualProduction/LevelSnapshots/Source/LevelSnapshots/Private/Restorability/SnapshotRestorability.cpp
@@ -105,7 +105,7 @@ bool UE::LevelSnapshots::Restorability::IsActorDesirableForCapture(const AActor*
 	for (const TSharedRef<ISnapshotRestorabilityOverrider>& Override : Overrides)
 	{
 		const ISnapshotRestorabilityOverrider::ERestorabilityOverride Result = Override->IsActorDesirableForCapture(Actor);
-		bSomebodyAllowed = Result == ISnapshotRestorabilityOverrider::ERestorabilityOverride::Allow;
+		bSomebodyAllowed |= Result == ISnapshotRestorabilityOverrider::ERestorabilityOverride::Allow;
 		if (Result == ISnapshotRestorabilityOverrider::ERestorabilityOverride::Disallow)
 		{
 			return false;
@@ -148,7 +148,7 @@ bool UE::LevelSnapshots::Restorability::IsComponentDesirableForCapture(const UAc
 	for (const TSharedRef<ISnapshotRestorabilityOverrider>& Override : Overrides)
 	{
 		const ISnapshotRestorabilityOverrider::ERestorabilityOverride Result = Override->IsComponentDesirableForCapture(Component);
-		bSomebodyAllowed = Result == ISnapshotRestorabilityOverrider::ERestorabilityOverride::Allow;
+		bSomebodyAllowed |= Result == ISnapshotRestorabilityOverrider::ERestorabilityOverride::Allow;
 		if (Result == ISnapshotRestorabilityOverrider::ERestorabilityOverride::Disallow)
 		{
 			return false;
