Fix unchecked properties sometimes not restored

diff --git a/Engine/Plugins/VirtualProduction/LevelSnapshots/Source/LevelSnapshotsEditor/Private/Views/Results/SLevelSnapshotsEditorResults.cpp b/Engine/Plugins/VirtualProduction/LevelSnapshots/Source/LevelSnapshotsEditor/Private/Views/Results/SLevelSnapshotsEditorResults.cpp
--- a/Engine/Plugins/VirtualProduction/LevelSnapshots/Source/LevelSnapshotsEditor/Private/Views/Results/SLevelSnapshotsEditorResults.cpp
+++ b/Engine/Plugins/VirtualProduction/LevelSnapshots/Source/LevelSnapshotsEditor/Private/Views/Results/SLevelSnapshotsEditorResults.cpp
@@ -608,7 +608,11 @@ void SLevelSnapshotsEditorResults::BuildSelectionSetFromSelectedPropertiesInEach
 					}
 				}
 
-				if (!CheckedNodeFieldPaths.IsEmpty())
+				if (CheckedNodeFieldPaths.IsEmpty())
+				{
+					SelectionMap.RemoveObjectPropertiesFromMap(WorldObject);
+				}
+				else
 				{
 					SelectionMap.AddObjectProperties(WorldObject, CheckedNodeFieldPaths);
 				}
