Remove usage of TStrongObjectPtr and minor refactors

diff --git a/Engine/Plugins/VirtualProduction/LevelSnapshots/Source/LevelSnapshots/Private/LevelSnapshotsFunctionLibrary.cpp b/Engine/Plugins/VirtualProduction/LevelSnapshots/Source/LevelSnapshots/Private/LevelSnapshotsFunctionLibrary.cpp
--- a/Engine/Plugins/VirtualProduction/LevelSnapshots/Source/LevelSnapshots/Private/LevelSnapshotsFunctionLibrary.cpp
+++ b/Engine/Plugins/VirtualProduction/LevelSnapshots/Source/LevelSnapshots/Private/LevelSnapshotsFunctionLibrary.cpp
@@ -61,8 +61,8 @@ void ULevelSnapshotsFunctionLibrary::ForEachMatchingCustomSubobjectPair(
 	ULevelSnapshot* Snapshot,
 	UObject* SnapshotRootObject,
 	UObject* WorldRootObject,
-	TFunction<void(UObject* SnapshotSubobject, UObject* EditorWorldSubobject)> HandleCustomSubobjectPair,
-	TFunction<void(UObject* UnmatchedSnapshotSubobject)> HandleUnmatchedSnapshotSubobject)
+	TFunctionRef<void(UObject* SnapshotSubobject, UObject* EditorWorldSubobject)> HandleCustomSubobjectPair,
+	TFunctionRef<void(UObject* UnmatchedSnapshotSubobject)> HandleUnmatchedSnapshotSubobject)
 {
 	UE::LevelSnapshots::Private::ForEachMatchingCustomSubobjectPair(Snapshot->GetSerializedData(), SnapshotRootObject, WorldRootObject, HandleCustomSubobjectPair, HandleUnmatchedSnapshotSubobject);
 }
diff --git a/Engine/Plugins/VirtualProduction/LevelSnapshots/Source/LevelSnapshots/Public/LevelSnapshotsFunctionLibrary.h b/Engine/Plugins/VirtualProduction/LevelSnapshots/Source/LevelSnapshots/Public/LevelSnapshotsFunctionLibrary.h
--- a/Engine/Plugins/VirtualProduction/LevelSnapshots/Source/LevelSnapshots/Public/LevelSnapshotsFunctionLibrary.h
+++ b/Engine/Plugins/VirtualProduction/LevelSnapshots/Source/LevelSnapshots/Public/LevelSnapshotsFunctionLibrary.h
@@ -53,7 +53,7 @@ public:
 		ULevelSnapshot* Snapshot, 
 		UObject* SnapshotRootObject,
 		UObject* WorldRootObject,
-		TFunction<void(UObject* SnapshotSubobject, UObject* EditorWorldSubobject)> HandleCustomSubobjectPair,
-		TFunction<void(UObject* UnmatchedSnapshotSubobject)> HandleUnmatchedSnapshotSubobject
+		TFunctionRef<void(UObject* SnapshotSubobject, UObject* EditorWorldSubobject)> HandleCustomSubobjectPair,
+		TFunctionRef<void(UObject* UnmatchedSnapshotSubobject)> HandleUnmatchedSnapshotSubobject
 		);
 };
\ No newline at end of file
diff --git a/Engine/Plugins/VirtualProduction/LevelSnapshots/Source/LevelSnapshotsEditor/Private/Data/LevelSnapshotsEditorData.cpp b/Engine/Plugins/VirtualProduction/LevelSnapshots/Source/LevelSnapshotsEditor/Private/Data/LevelSnapshotsEditorData.cpp
--- a/Engine/Plugins/VirtualProduction/LevelSnapshots/Source/LevelSnapshotsEditor/Private/Data/LevelSnapshotsEditorData.cpp
+++ b/Engine/Plugins/VirtualProduction/LevelSnapshots/Source/LevelSnapshotsEditor/Private/Data/LevelSnapshotsEditorData.cpp
@@ -91,36 +91,24 @@ void ULevelSnapshotsEditorData::BeginDestroy()
 void ULevelSnapshotsEditorData::CleanupAfterEditorClose()
 {
 	OnActiveSnapshotChanged.Clear();
-	OnEditedFiterChanged.Clear();
+	OnEditedFilterChanged.Clear();
 	OnUserDefinedFiltersChanged.Clear();
 
-	ActiveSnapshot.Reset();
-	EditedFilter.Reset();
+	ActiveSnapshot = nullptr;
+	EditedFilter = nullptr;
 
 	FilterResults->CleanReferences();
 }
 
-void ULevelSnapshotsEditorData::SetActiveSnapshot(const TOptional<ULevelSnapshot*>& NewActiveSnapshot)
+void ULevelSnapshotsEditorData::SetActiveSnapshot(ULevelSnapshot* NewActiveSnapshot)
 {
 	SCOPED_SNAPSHOT_EDITOR_TRACE(SetActiveSnapshot);
 	
-	ActiveSnapshot = NewActiveSnapshot.Get(nullptr) ? TStrongObjectPtr<ULevelSnapshot>(NewActiveSnapshot.GetValue()) : TOptional<TStrongObjectPtr<ULevelSnapshot>>();
-	FilterResults->SetActiveLevelSnapshot(NewActiveSnapshot.Get(nullptr));
+	ActiveSnapshot = NewActiveSnapshot;
+	FilterResults->SetActiveLevelSnapshot(NewActiveSnapshot);
 	OnActiveSnapshotChanged.Broadcast(GetActiveSnapshot());
 }
 
-void ULevelSnapshotsEditorData::ClearActiveSnapshot()
-{
-	ActiveSnapshot.Reset();
-	FilterResults->SetActiveLevelSnapshot(nullptr);
-	OnActiveSnapshotChanged.Broadcast(TOptional<ULevelSnapshot*>(nullptr));
-}
-
-TOptional<ULevelSnapshot*> ULevelSnapshotsEditorData::GetActiveSnapshot() const
-{
-	return ActiveSnapshot.IsSet() ? TOptional<ULevelSnapshot*>(ActiveSnapshot->Get()) : TOptional<ULevelSnapshot*>();
-}
-
 UWorld* ULevelSnapshotsEditorData::GetEditorWorld()
 {
 	// If this function is called very early during startup, the initial editor GWorld may not have been created yet!
@@ -135,20 +123,13 @@ UWorld* ULevelSnapshotsEditorData::GetEditorWorld()
 	return nullptr;
 }
 
-void ULevelSnapshotsEditorData::SetEditedFilter(const TOptional<UNegatableFilter*>& InFilter)
-{
-	EditedFilter = InFilter.Get(nullptr) ? TStrongObjectPtr<UNegatableFilter>(InFilter.GetValue()) : TOptional<TStrongObjectPtr<UNegatableFilter>>();
-	OnEditedFiterChanged.Broadcast(GetEditedFilter());
-}
-
-TOptional<UNegatableFilter*> ULevelSnapshotsEditorData::GetEditedFilter() const
+void ULevelSnapshotsEditorData::SetEditedFilter(UNegatableFilter* InFilter)
 {
-	return EditedFilter.IsSet() ? TOptional<UNegatableFilter*>(EditedFilter->Get()) : TOptional<UNegatableFilter*>();
-}
-
-bool ULevelSnapshotsEditorData::IsEditingFilter(UNegatableFilter* Filter) const
-{
-	return (Filter == nullptr && !EditedFilter.IsSet()) || (EditedFilter.IsSet() && Filter && Filter == EditedFilter->Get()); 
+	if (EditedFilter != InFilter)
+	{
+		EditedFilter = InFilter;
+		OnEditedFilterChanged.Broadcast(GetEditedFilter());
+	}
 }
 
 UFavoriteFilterContainer* ULevelSnapshotsEditorData::GetFavoriteFilters() const
diff --git a/Engine/Plugins/VirtualProduction/LevelSnapshots/Source/LevelSnapshotsEditor/Private/Data/LevelSnapshotsEditorData.h b/Engine/Plugins/VirtualProduction/LevelSnapshots/Source/LevelSnapshotsEditor/Private/Data/LevelSnapshotsEditorData.h
--- a/Engine/Plugins/VirtualProduction/LevelSnapshots/Source/LevelSnapshotsEditor/Private/Data/LevelSnapshotsEditorData.h
+++ b/Engine/Plugins/VirtualProduction/LevelSnapshots/Source/LevelSnapshotsEditor/Private/Data/LevelSnapshotsEditorData.h
@@ -35,12 +35,13 @@ public:
 
 	
 	/******************** Active snapshot ********************/
-	
-	void SetActiveSnapshot(const TOptional<ULevelSnapshot*>& NewActiveSnapshot);
-	void ClearActiveSnapshot();
-	TOptional<ULevelSnapshot*> GetActiveSnapshot() const;
 
-	DECLARE_EVENT_OneParam(ULevelSnapshotsEditorData, FOnActiveSnapshotChanged, const TOptional<ULevelSnapshot*>& /* NewSnapshot */);
+	/** @param NewActiveSnapshot Can be nullptr */
+	void SetActiveSnapshot(ULevelSnapshot* NewActiveSnapshot);
+	void ClearActiveSnapshot() { SetActiveSnapshot(nullptr); }
+	ULevelSnapshot* GetActiveSnapshot() const { return ActiveSnapshot; }
+
+	DECLARE_EVENT_OneParam(ULevelSnapshotsEditorData, FOnActiveSnapshotChanged, ULevelSnapshot* /* NewSnapshot */);
 	FOnActiveSnapshotChanged OnActiveSnapshotChanged;
 
 
@@ -52,12 +53,12 @@ public:
 	
 	/******************** Edited filter ********************/
 	
-	void SetEditedFilter(const TOptional<UNegatableFilter*>& InFilter);
-	TOptional<UNegatableFilter*> GetEditedFilter() const;
-	bool IsEditingFilter(UNegatableFilter* Filter) const;
+	void SetEditedFilter(UNegatableFilter* InFilter);
+	UNegatableFilter* GetEditedFilter() const { return EditedFilter; }
+	bool IsEditingFilter(UNegatableFilter* Filter) const { return Filter == GetEditedFilter(); }
 
-	DECLARE_EVENT_OneParam(ULevelSnapshotsEditorData, FOnEditedFilterChanged, const TOptional<UNegatableFilter*>& /*NewEditedFilter*/);
-	FOnEditedFilterChanged OnEditedFiterChanged;
+	DECLARE_EVENT_OneParam(ULevelSnapshotsEditorData, FOnEditedFilterChanged, UNegatableFilter* /*NewEditedFilter*/);
+	FOnEditedFilterChanged OnEditedFilterChanged;
 
 	DECLARE_EVENT(ULevelSnapshotsEditorData, FOnRefreshResults);
 	FOnRefreshResults OnRefreshResults;
@@ -107,7 +108,9 @@ private:
 	TObjectPtr<UFilteredResults> FilterResults;
 
 	/* Snapshot selected by user */
-	TOptional<TStrongObjectPtr<ULevelSnapshot>> ActiveSnapshot;
+	UPROPERTY()
+	TObjectPtr<ULevelSnapshot> ActiveSnapshot;
 	/* Filter visible in details panel */
-	TOptional<TStrongObjectPtr<UNegatableFilter>> EditedFilter;
+	UPROPERTY()
+	TObjectPtr<UNegatableFilter> EditedFilter;
 };
\ No newline at end of file
diff --git a/Engine/Plugins/VirtualProduction/LevelSnapshots/Source/LevelSnapshotsEditor/Private/Views/Filter/SLevelSnapshotsEditorFilters.cpp b/Engine/Plugins/VirtualProduction/LevelSnapshots/Source/LevelSnapshotsEditor/Private/Views/Filter/SLevelSnapshotsEditorFilters.cpp
--- a/Engine/Plugins/VirtualProduction/LevelSnapshots/Source/LevelSnapshotsEditor/Private/Views/Filter/SLevelSnapshotsEditorFilters.cpp
+++ b/Engine/Plugins/VirtualProduction/LevelSnapshots/Source/LevelSnapshotsEditor/Private/Views/Filter/SLevelSnapshotsEditorFilters.cpp
@@ -79,7 +79,7 @@ SLevelSnapshotsEditorFilters::~SLevelSnapshotsEditorFilters()
 	if (ULevelSnapshotsEditorData* Data = GetEditorData())
 	{
 		Data->OnUserDefinedFiltersChanged.Remove(OnUserDefinedFiltersChangedHandle);
-		Data->OnEditedFiterChanged.Remove(OnEditedFilterChangedHandle);
+		Data->OnEditedFilterChanged.Remove(OnEditedFilterChangedHandle);
 		Data->GetUserDefinedFilters()->OnFilterModified.Remove(OnFilterModifiedHandle);
 	}
 }
@@ -223,11 +223,11 @@ void SLevelSnapshotsEditorFilters::Construct(const FArguments& InArgs, ULevelSna
 		}
 		OnFilterModifiedHandle = NewFilter->OnFilterModified.AddRaw(this, &SLevelSnapshotsEditorFilters::OnFilterModified);
 		
-		GetEditorData()->SetEditedFilter(TOptional<UNegatableFilter*>());
+		GetEditorData()->SetEditedFilter(nullptr);
 		RefreshGroups();
 	});
 
-	OnEditedFilterChangedHandle = GetEditorData()->OnEditedFiterChanged.AddLambda([this](const TOptional<UNegatableFilter*>& ActiveFilter)
+	OnEditedFilterChangedHandle = GetEditorData()->OnEditedFilterChanged.AddLambda([this](const TOptional<UNegatableFilter*>& ActiveFilter)
 	{
 		FilterDetailsView->SetObject(ActiveFilter.IsSet() ? ActiveFilter.GetValue() : nullptr);
 	});
diff --git a/Engine/Plugins/VirtualProduction/LevelSnapshots/Source/LevelSnapshotsEditor/Private/Views/Results/LevelSnapshotsEditorResultsRow.cpp b/Engine/Plugins/VirtualProduction/LevelSnapshots/Source/LevelSnapshotsEditor/Private/Views/Results/LevelSnapshotsEditorResultsRow.cpp
--- a/Engine/Plugins/VirtualProduction/LevelSnapshots/Source/LevelSnapshotsEditor/Private/Views/Results/LevelSnapshotsEditorResultsRow.cpp
+++ b/Engine/Plugins/VirtualProduction/LevelSnapshots/Source/LevelSnapshotsEditor/Private/Views/Results/LevelSnapshotsEditorResultsRow.cpp
@@ -348,10 +348,10 @@ void FLevelSnapshotsEditorResultsRow::GenerateModifiedActorGroupChildren(FProper
 	}
 
 	// Custom subobjects handled by external ICustomObjectSnapshotSerializer implementation
-	TMap<UObject*, UObject*> WorldToSnapshotCustomSubobjects;
-	
-	ULevelSnapshot* ActiveSnapshot = ResultsViewPtr.IsValid() && ResultsViewPtr.Pin()->GetEditorDataPtr() && ResultsViewPtr.Pin()->GetEditorDataPtr()->GetActiveSnapshot().IsSet() ? 
-		ResultsViewPtr.Pin()->GetEditorDataPtr()->GetActiveSnapshot().GetValue() : nullptr;
+	const TSharedPtr<SLevelSnapshotsEditorResults> ResultsViewPin = ResultsViewPtr.Pin();
+	ULevelSnapshot* ActiveSnapshot = ResultsViewPin && ResultsViewPin->GetEditorDataPtr()->GetActiveSnapshot()
+		? ResultsViewPtr.Pin()->GetEditorDataPtr()->GetActiveSnapshot()
+		: nullptr;
 
 	const TFunction<void(UObject* UnmatchedSnapshotSubobject)> HandleUnmatchedCustomSnapshotSubobject = [](UObject* UnmatchedCustomSnapshotSubobject)
 	{
@@ -359,6 +359,7 @@ void FLevelSnapshotsEditorResultsRow::GenerateModifiedActorGroupChildren(FProper
 		// If the row is checked when the snapshot is applied, call FPropertySelectionMap::AddCustomEditorSubobjectToRecreate
 	};
 
+	TMap<UObject*, UObject*> WorldToSnapshotCustomSubobjects;
 	const TFunction<void(UObject* SnapshotSubobject, UObject* EditorWorldSubobject)> HandleMatchedCustomSubobjectPair = [&WorldToSnapshotCustomSubobjects, ActiveSnapshot, &HandleMatchedCustomSubobjectPair, &HandleUnmatchedCustomSnapshotSubobject](UObject* SnapshotSubobject, UObject* EditorWorldSubobject)
 	{
 		check(EditorWorldSubobject);
diff --git a/Engine/Plugins/VirtualProduction/LevelSnapshots/Source/LevelSnapshotsEditor/Private/Views/Results/SLevelSnapshotsEditorResults.cpp b/Engine/Plugins/VirtualProduction/LevelSnapshots/Source/LevelSnapshotsEditor/Private/Views/Results/SLevelSnapshotsEditorResults.cpp
--- a/Engine/Plugins/VirtualProduction/LevelSnapshots/Source/LevelSnapshotsEditor/Private/Views/Results/SLevelSnapshotsEditorResults.cpp
+++ b/Engine/Plugins/VirtualProduction/LevelSnapshots/Source/LevelSnapshotsEditor/Private/Views/Results/SLevelSnapshotsEditorResults.cpp
@@ -332,9 +332,9 @@ TOptional<ULevelSnapshot*> SLevelSnapshotsEditorResults::GetSelectedLevelSnapsho
 	return ensure(EditorDataPtr.IsValid()) ? EditorDataPtr->GetActiveSnapshot() : TOptional<ULevelSnapshot*>();
 }
 
-void SLevelSnapshotsEditorResults::OnSnapshotSelected(const TOptional<ULevelSnapshot*>& InLevelSnapshot)
+void SLevelSnapshotsEditorResults::OnSnapshotSelected(ULevelSnapshot* InLevelSnapshot)
 {	
-	if (InLevelSnapshot.IsSet() && InLevelSnapshot.GetValue())
+	if (InLevelSnapshot)
 	{
 		UpdateSnapshotNameText(InLevelSnapshot);
 		
@@ -1009,25 +1009,25 @@ bool SLevelSnapshotsEditorResults::GenerateTreeViewChildren_AddedActors(FLevelSn
 
 bool SLevelSnapshotsEditorResults::GenerateTreeViewChildren_RemovedActors(FLevelSnapshotsEditorResultsRowPtr RemovedActorsHeader)
 {
-	check(RemovedActorsHeader);
-	check (EditorDataPtr.IsValid());
-	check (EditorDataPtr->GetActiveSnapshot().IsSet());
-	
-	TObjectPtr<ULevelSnapshot> ActiveSnapshot = EditorDataPtr->GetActiveSnapshot().GetValue();
-	check(ActiveSnapshot);
+	if (!ensure(RemovedActorsHeader && EditorDataPtr.IsValid() && EditorDataPtr->GetActiveSnapshot() != nullptr))
+	{
+		return false;
+	}
 	
+	const TObjectPtr<ULevelSnapshot> ActiveSnapshot = EditorDataPtr->GetActiveSnapshot();
 	for (const FSoftObjectPath& ActorPath : FilterListData.GetRemovedOriginalActorPaths_AllowedByFilter())
 	{
-		FString ActorName = ActiveSnapshot->GetActorLabel(ActorPath);
-		
-		// Create group
+		const FString ActorName = ActiveSnapshot->GetActorLabel(ActorPath);
 		FLevelSnapshotsEditorResultsRowPtr NewActorRow = MakeShared<FLevelSnapshotsEditorResultsRow>(
-			FLevelSnapshotsEditorResultsRow(
-				FText::FromString(ActorName), FLevelSnapshotsEditorResultsRow::RemovedActorToAdd, ECheckBoxState::Checked, SharedThis(this), RemovedActorsHeader));
+			FText::FromString(ActorName),
+			FLevelSnapshotsEditorResultsRow::RemovedActorToAdd,
+			ECheckBoxState::Checked,
+			SharedThis(this),
+			RemovedActorsHeader
+		);
+		
 		NewActorRow->InitRemovedActorRow(ActorPath);
-
 		RemovedActorsHeader->AddToChildRows(NewActorRow);
-
 		TreeViewRemovedActorGroupObjects.Add(NewActorRow);
 	}
 
diff --git a/Engine/Plugins/VirtualProduction/LevelSnapshots/Source/LevelSnapshotsEditor/Private/Views/Results/SLevelSnapshotsEditorResults.h b/Engine/Plugins/VirtualProduction/LevelSnapshots/Source/LevelSnapshotsEditor/Private/Views/Results/SLevelSnapshotsEditorResults.h
--- a/Engine/Plugins/VirtualProduction/LevelSnapshots/Source/LevelSnapshotsEditor/Private/Views/Results/SLevelSnapshotsEditorResults.h
+++ b/Engine/Plugins/VirtualProduction/LevelSnapshots/Source/LevelSnapshotsEditor/Private/Views/Results/SLevelSnapshotsEditorResults.h
@@ -49,7 +49,7 @@ public:
 
 	TOptional<ULevelSnapshot*> GetSelectedLevelSnapshot() const;
 
-	void OnSnapshotSelected(const TOptional<ULevelSnapshot*>& InLevelSnapshot);
+	void OnSnapshotSelected(ULevelSnapshot* InLevelSnapshot);
 	void RefreshResults(const bool bSnapshotHasChanged = false);
 	FReply OnClickApplyToWorld();
 
diff --git a/Engine/Plugins/VirtualProduction/LevelSnapshots/Source/LevelSnapshotsEditor/Private/Widgets/SLevelSnapshotsEditorFilter.cpp b/Engine/Plugins/VirtualProduction/LevelSnapshots/Source/LevelSnapshotsEditor/Private/Widgets/SLevelSnapshotsEditorFilter.cpp
--- a/Engine/Plugins/VirtualProduction/LevelSnapshots/Source/LevelSnapshotsEditor/Private/Widgets/SLevelSnapshotsEditorFilter.cpp
+++ b/Engine/Plugins/VirtualProduction/LevelSnapshots/Source/LevelSnapshotsEditor/Private/Widgets/SLevelSnapshotsEditorFilter.cpp
@@ -45,7 +45,7 @@ SLevelSnapshotsEditorFilter::~SLevelSnapshotsEditorFilter()
 {
 	if (ensure(EditorData.IsValid()))
 	{
-		EditorData->OnEditedFiterChanged.Remove(ActiveFilterChangedDelegateHandle);
+		EditorData->OnEditedFilterChanged.Remove(ActiveFilterChangedDelegateHandle);
 	}
 
 	if (SnapshotFilter.IsValid())
@@ -125,7 +125,7 @@ void SLevelSnapshotsEditorFilter::Construct(const FArguments& InArgs, const TWea
 	// Hightlight & unhighlight filter when being edited
 	FilterNamePtr->SetOnClicked(FOnClicked::CreateRaw(this, &SLevelSnapshotsEditorFilter::OnSelectFilterForEdit));
 	
-	ActiveFilterChangedDelegateHandle = InEditorData->OnEditedFiterChanged.AddRaw(this, &SLevelSnapshotsEditorFilter::OnActiveFilterChanged);
+	ActiveFilterChangedDelegateHandle = InEditorData->OnEditedFilterChanged.AddSP(this, &SLevelSnapshotsEditorFilter::OnActiveFilterChanged);
 	OnFilterDestroyedDelegateHandle = InFilter->OnFilterDestroyed.AddLambda([this](UNegatableFilter* Filter)
 	{
 		if (EditorData.IsValid() && EditorData->IsEditingFilter(Filter))
@@ -211,7 +211,7 @@ FReply SLevelSnapshotsEditorFilter::OnSelectFilterForEdit()
 	if (EditorData->IsEditingFilter(SnapshotFilter.Get()))
 	{
 		bShouldHighlightFilter = false;
-		EditorData->SetEditedFilter(TOptional<UNegatableFilter*>());
+		EditorData->SetEditedFilter(nullptr);
 	}
 	else
 	{
@@ -241,13 +241,13 @@ FReply SLevelSnapshotsEditorFilter::OnRemoveFilter()
 	return FReply::Handled();
 }
 
-void SLevelSnapshotsEditorFilter::OnActiveFilterChanged(const TOptional<UNegatableFilter*>& NewFilter)
+void SLevelSnapshotsEditorFilter::OnActiveFilterChanged(UNegatableFilter* NewFilter)
 {
-	if (!SnapshotFilter.IsValid()) // This can actually become stale after a save: UI rebuilds next tick but object was already destroyed.
+	if (SnapshotFilter.IsValid()) // This can actually become stale after a save: UI rebuilds next tick but object was already destroyed.
 	{
-		return;
+		bShouldHighlightFilter = NewFilter && NewFilter == SnapshotFilter.Get();
 	}
-	bShouldHighlightFilter = NewFilter.IsSet() && *NewFilter == SnapshotFilter.Get();
+	
 }
 
 
diff --git a/Engine/Plugins/VirtualProduction/LevelSnapshots/Source/LevelSnapshotsEditor/Private/Widgets/SLevelSnapshotsEditorFilter.h b/Engine/Plugins/VirtualProduction/LevelSnapshots/Source/LevelSnapshotsEditor/Private/Widgets/SLevelSnapshotsEditorFilter.h
--- a/Engine/Plugins/VirtualProduction/LevelSnapshots/Source/LevelSnapshotsEditor/Private/Widgets/SLevelSnapshotsEditorFilter.h
+++ b/Engine/Plugins/VirtualProduction/LevelSnapshots/Source/LevelSnapshotsEditor/Private/Widgets/SLevelSnapshotsEditorFilter.h
@@ -51,7 +51,7 @@ private:
 	FReply OnNegateFilter();
 	FReply OnRemoveFilter();
 	
-	void OnActiveFilterChanged(const TOptional<UNegatableFilter*>& NewFilter);
+	void OnActiveFilterChanged(UNegatableFilter* NewFilter);
 
 	/* Used to remove this filter */
 	FOnClickRemoveFilter OnClickRemoveFilter;
