Fix Level Snapshot not restoring original actors if they were converted to a Level Instance/Packed Level Instance

diff --git a/Engine/Plugins/VirtualProduction/LevelSnapshots/Source/LevelSnapshots/Private/Data/LevelSnapshot.cpp b/Engine/Plugins/VirtualProduction/LevelSnapshots/Source/LevelSnapshots/Private/Data/LevelSnapshot.cpp
--- a/Engine/Plugins/VirtualProduction/LevelSnapshots/Source/LevelSnapshots/Private/Data/LevelSnapshot.cpp
+++ b/Engine/Plugins/VirtualProduction/LevelSnapshots/Source/LevelSnapshots/Private/Data/LevelSnapshot.cpp
@@ -32,7 +32,7 @@
 void ULevelSnapshot::ApplySnapshotToWorld(UWorld* TargetWorld, const FPropertySelectionMap& SelectionSet)
 {
 	SCOPED_SNAPSHOT_CORE_TRACE(ApplyToWorld);
-	if (TargetWorld == nullptr)
+	if (TargetWorld == nullptr || !ensure(IsInGameThread()))
 	{
 		return;
 	}
@@ -46,7 +46,9 @@ void ULevelSnapshot::ApplySnapshotToWorld(UWorld* TargetWorld, const FPropertySe
 	
 	EnsureWorldInitialised();
 	
+	OnPreApplySnapshotDelegate.Broadcast();
 	UE::LevelSnapshots::Private::ApplyToWorld(SerializedData, Cache, TargetWorld, GetPackage(), SelectionSet);
+	OnPostApplySnapshotDelegate.Broadcast();
 }
 
 bool ULevelSnapshot::SnapshotWorld(UWorld* TargetWorld)
diff --git a/Engine/Plugins/VirtualProduction/LevelSnapshots/Source/LevelSnapshots/Public/Data/LevelSnapshot.h b/Engine/Plugins/VirtualProduction/LevelSnapshots/Source/LevelSnapshots/Public/Data/LevelSnapshot.h
--- a/Engine/Plugins/VirtualProduction/LevelSnapshots/Source/LevelSnapshots/Public/Data/LevelSnapshot.h
+++ b/Engine/Plugins/VirtualProduction/LevelSnapshots/Source/LevelSnapshots/Public/Data/LevelSnapshot.h
@@ -30,6 +30,7 @@ public:
 
 	DECLARE_DELEGATE_OneParam(FActorPathConsumer, const FSoftObjectPath& /*OriginalActorPath*/);
 	DECLARE_DELEGATE_OneParam(FActorConsumer, AActor* /*WorldActor*/);
+	DECLARE_MULTICAST_DELEGATE(FSnapshotEvent);
 	
 	/* Captures the current state of the given world. */
 	bool SnapshotWorld(UWorld* TargetWorld);
@@ -75,7 +76,6 @@ public:
 	void DiffWorld(UWorld* World, FActorConsumer HandleMatchedActor, FActorPathConsumer HandleRemovedActor, FActorConsumer HandleAddedActor);
 
 	
-	
 	/* Sets the name of this snapshot. */
 	UFUNCTION(BlueprintCallable, Category = "Level Snapshots")
 	void SetSnapshotName(const FName& InSnapshotName);
@@ -94,6 +94,9 @@ public:
 	const FWorldSnapshotData& GetSerializedData() const { return SerializedData; }
 	const FSnapshotDataCache& GetCache() const { return Cache; }
 
+	FSnapshotEvent& OnPreApplySnapshot() { return OnPreApplySnapshotDelegate; }
+	FSnapshotEvent& OnPostApplySnapshot() { return OnPostApplySnapshotDelegate; }
+
 #if !WITH_EDITOR
 	/**
 	 * Indicates that user code is aware that ClearCachedDiffFlag must be called manually.
@@ -182,4 +185,7 @@ private:
 	/** User defined description of the snapshot */
 	UPROPERTY(AssetRegistrySearchable, BlueprintGetter = "GetSnapshotDescription", EditAnywhere, Category = "Level Snapshots")
 	FString SnapshotDescription;
+
+	FSnapshotEvent OnPreApplySnapshotDelegate;
+	FSnapshotEvent OnPostApplySnapshotDelegate;
 };
diff --git a/Engine/Plugins/VirtualProduction/LevelSnapshots/Source/LevelSnapshotsEditor/Private/Data/LevelSnapshotsEditorData.cpp b/Engine/Plugins/VirtualProduction/LevelSnapshots/Source/LevelSnapshotsEditor/Private/Data/LevelSnapshotsEditorData.cpp
--- a/Engine/Plugins/VirtualProduction/LevelSnapshots/Source/LevelSnapshotsEditor/Private/Data/LevelSnapshotsEditorData.cpp
+++ b/Engine/Plugins/VirtualProduction/LevelSnapshots/Source/LevelSnapshotsEditor/Private/Data/LevelSnapshotsEditorData.cpp
@@ -103,8 +103,25 @@ void ULevelSnapshotsEditorData::CleanupAfterEditorClose()
 void ULevelSnapshotsEditorData::SetActiveSnapshot(ULevelSnapshot* NewActiveSnapshot)
 {
 	SCOPED_SNAPSHOT_EDITOR_TRACE(SetActiveSnapshot);
+
+	if (bIsApplyingSnapshot)
+	{
+		return;
+	}
+
+	if (ActiveSnapshot)
+	{
+		ActiveSnapshot->OnPreApplySnapshot().RemoveAll(this);
+		ActiveSnapshot->OnPostApplySnapshot().RemoveAll(this);
+	}
 	
 	ActiveSnapshot = NewActiveSnapshot;
+	if (ActiveSnapshot)
+	{
+		ActiveSnapshot->OnPreApplySnapshot().AddUObject(this, &ULevelSnapshotsEditorData::OnPreApplySnapshot);
+		ActiveSnapshot->OnPostApplySnapshot().AddUObject(this, &ULevelSnapshotsEditorData::OnPostApplySnapshot);
+	}
+	
 	FilterResults->SetActiveLevelSnapshot(NewActiveSnapshot);
 	OnActiveSnapshotChanged.Broadcast(GetActiveSnapshot());
 }
diff --git a/Engine/Plugins/VirtualProduction/LevelSnapshots/Source/LevelSnapshotsEditor/Private/Data/LevelSnapshotsEditorData.h b/Engine/Plugins/VirtualProduction/LevelSnapshots/Source/LevelSnapshotsEditor/Private/Data/LevelSnapshotsEditorData.h
--- a/Engine/Plugins/VirtualProduction/LevelSnapshots/Source/LevelSnapshotsEditor/Private/Data/LevelSnapshotsEditorData.h
+++ b/Engine/Plugins/VirtualProduction/LevelSnapshots/Source/LevelSnapshotsEditor/Private/Data/LevelSnapshotsEditorData.h
@@ -103,6 +103,10 @@ private:
 	UPROPERTY()
 	bool bIsFilterDirty = false;
 
+	/** Whether ActiveSnapshot is currently being restored. Does not allow changing the active snapshot while true.  */
+	UPROPERTY()
+	bool bIsApplyingSnapshot = false;
+
 	/* Converts UserDefinedFilters into ULevelSnapshotsSelectionSet display in results view. */
 	UPROPERTY()
 	TObjectPtr<UFilteredResults> FilterResults;
@@ -113,4 +117,7 @@ private:
 	/* Filter visible in details panel */
 	UPROPERTY()
 	TObjectPtr<UNegatableFilter> EditedFilter;
-};
\ No newline at end of file
+
+	void OnPreApplySnapshot() { bIsApplyingSnapshot = true; }
+	void OnPostApplySnapshot() { bIsApplyingSnapshot = false; }
+};
