Fix unit test

diff --git a/Engine/Plugins/VirtualProduction/LevelSnapshots/Source/LevelSnapshots/Private/Tests/SubobjectTests.cpp b/Engine/Plugins/VirtualProduction/LevelSnapshots/Source/LevelSnapshots/Private/Tests/SubobjectTests.cpp
--- a/Engine/Plugins/VirtualProduction/LevelSnapshots/Source/LevelSnapshots/Private/Tests/SubobjectTests.cpp
+++ b/Engine/Plugins/VirtualProduction/LevelSnapshots/Source/LevelSnapshots/Private/Tests/SubobjectTests.cpp
@@ -715,10 +715,10 @@ namespace UE::LevelSnapshots::Private::Tests
 		ILevelSnapshotsModule::Get().RegisterRestorabilityOverrider(Support);
 		ILevelSnapshotsModule::Get().AddSkippedSubobjectClasses({ USubobject::StaticClass() });
 		ON_SCOPE_EXIT
-{
-	ILevelSnapshotsModule::Get().UnregisterRestorabilityOverrider(Support);
-	ILevelSnapshotsModule::Get().RemoveSkippedSubobjectClasses({ USubobject::StaticClass() });
-};
+		{
+			ILevelSnapshotsModule::Get().UnregisterRestorabilityOverrider(Support);
+			ILevelSnapshotsModule::Get().RemoveSkippedSubobjectClasses({ USubobject::StaticClass() });
+		};
 	
 		ASnapshotTestActor* ActorA = nullptr;
 		ASnapshotTestActor* ActorB = nullptr;
@@ -799,9 +799,10 @@ namespace UE::LevelSnapshots::Private::Tests
 					TestTrue(TEXT("Empty snapshot collection"), ActorSelection->IsPropertySelected(nullptr, Property));
 				}
 
-				// + 2 because FSoftObjectPath has 2 properties; they are implicitly added because ASnapshotTestActor::SoftPath is technically a struct
-				const int32 ExpectedNumProperties = SoftObjectProperties.Num() + 2 + EmptySnapshotCollections.Num();
-				TestEqual(TEXT("ActorB only has soft object paths in selection set"), ActorSelection->GetSelectedProperties().Num(), ExpectedNumProperties);
+				// + 2 because FSoftObjectPath has 4 properties; they are implicitly added because ASnapshotTestActor::SoftPath is technically a struct
+				const int32 ExpectedNumProperties = SoftObjectProperties.Num() + 4 + EmptySnapshotCollections.Num();
+				const TArray<FLevelSnapshotPropertyChain>& SelectedProperties = ActorSelection->GetSelectedProperties();
+				TestEqual(TEXT("ActorB only has soft object paths in selection set"), SelectedProperties.Num(), ExpectedNumProperties);
 			})
 
 			// Because USubobject is skipped, the collections should contain null
