Fix not being able to select plugin path for saving snapshots and minor code clean-up

diff --git a/Engine/Plugins/VirtualProduction/LevelSnapshots/Source/LevelSnapshotsEditor/Private/LevelSnapshotsEditorModule.cpp b/Engine/Plugins/VirtualProduction/LevelSnapshots/Source/LevelSnapshotsEditor/Private/LevelSnapshotsEditorModule.cpp
--- a/Engine/Plugins/VirtualProduction/LevelSnapshots/Source/LevelSnapshotsEditor/Private/LevelSnapshotsEditorModule.cpp
+++ b/Engine/Plugins/VirtualProduction/LevelSnapshots/Source/LevelSnapshotsEditor/Private/LevelSnapshotsEditorModule.cpp
@@ -58,8 +58,8 @@ void FLevelSnapshotsEditorModule::StartupModule()
 			DataMangementSettingsSectionPtr = SettingsModule.RegisterSettings("Project", "Plugins", "Level Snapshots Editor",
 				NSLOCTEXT("LevelSnapshots", "LevelSnapshotsEditorSettingsCategoryDisplayName", "Level Snapshots Editor"),
 				NSLOCTEXT("LevelSnapshots", "LevelSnapshotsEditorSettingsDescription", "Configure the Level Snapshots Editor settings"),
-				GetMutableDefault<ULevelSnapshotsEditorSettings>());
-			DataMangementSettingsSectionPtr->OnModified().BindRaw(this, &FLevelSnapshotsEditorModule::HandleModifiedProjectSettings);
+				GetMutableDefault<ULevelSnapshotsEditorSettings>()
+				);
 		}
 	});
 }
@@ -148,16 +148,6 @@ void FLevelSnapshotsEditorModule::OpenSnapshotsEditor()
 	FGlobalTabmanager::Get()->TryInvokeTab(LevelSnapshotsEditor::LevelSnapshotsTabName);
 }
 
-bool FLevelSnapshotsEditorModule::HandleModifiedProjectSettings()
-{
-	ULevelSnapshotsEditorSettings* Settings = ULevelSnapshotsEditorSettings::GetMutable();
-	Settings->ValidateRootLevelSnapshotSaveDirAsGameContentRelative();
-	Settings->SanitizeAllProjectSettingsPaths(true);
-		
-	Settings->SaveConfig();
-	return true;
-}
-
 void FLevelSnapshotsEditorModule::RegisterEditorToolbar()
 {
 	if (IsRunningGame())
diff --git a/Engine/Plugins/VirtualProduction/LevelSnapshots/Source/LevelSnapshotsEditor/Private/LevelSnapshotsEditorModule.h b/Engine/Plugins/VirtualProduction/LevelSnapshots/Source/LevelSnapshotsEditor/Private/LevelSnapshotsEditorModule.h
--- a/Engine/Plugins/VirtualProduction/LevelSnapshots/Source/LevelSnapshotsEditor/Private/LevelSnapshotsEditorModule.h
+++ b/Engine/Plugins/VirtualProduction/LevelSnapshots/Source/LevelSnapshotsEditor/Private/LevelSnapshotsEditorModule.h
@@ -29,7 +29,7 @@ public:
 	//~ End IModuleInterface Interface
 
 	static bool GetUseCreationForm() { return ULevelSnapshotsEditorSettings::Get()->bUseCreationForm; }
-	static void SetUseCreationForm(bool bInUseCreationForm) { ULevelSnapshotsEditorSettings::GetMutable()->bUseCreationForm = bInUseCreationForm; }
+	static void SetUseCreationForm(bool bInUseCreationForm) { ULevelSnapshotsEditorSettings::Get()->bUseCreationForm = bInUseCreationForm; }
 	static void ToggleUseCreationForm() { SetUseCreationForm(!GetUseCreationForm()); }
 
 	void OpenLevelSnapshotsDialogWithAssetSelected(const FAssetData& InAssetData);
@@ -43,8 +43,6 @@ private:
 	TSharedRef<SDockTab> SpawnLevelSnapshotsTab(const FSpawnTabArgs& SpawnTabArgs);
 	ULevelSnapshotsEditorData* AllocateTransientPreset();
 	
-	bool HandleModifiedProjectSettings();
-	
 	void RegisterEditorToolbar();
 	void MapEditorToolbarActions();
 	TSharedRef<SWidget> FillEditorToolbarComboButtonMenuOptions(TSharedPtr<class FUICommandList> Commands);
diff --git a/Engine/Plugins/VirtualProduction/LevelSnapshots/Source/LevelSnapshotsEditor/Private/LevelSnapshotsEditorSettings.cpp b/Engine/Plugins/VirtualProduction/LevelSnapshots/Source/LevelSnapshotsEditor/Private/LevelSnapshotsEditorSettings.cpp
--- a/Engine/Plugins/VirtualProduction/LevelSnapshots/Source/LevelSnapshotsEditor/Private/LevelSnapshotsEditorSettings.cpp
+++ b/Engine/Plugins/VirtualProduction/LevelSnapshots/Source/LevelSnapshotsEditor/Private/LevelSnapshotsEditorSettings.cpp
@@ -6,35 +6,21 @@
 #include "Application/SlateApplicationBase.h"
 #include "HAL/PlatformApplicationMisc.h"
 
-const ULevelSnapshotsEditorSettings* ULevelSnapshotsEditorSettings::Get()
-{
-	return GetDefault<ULevelSnapshotsEditorSettings>();
-}
-
-ULevelSnapshotsEditorSettings* ULevelSnapshotsEditorSettings::GetMutable()
+ULevelSnapshotsEditorSettings* ULevelSnapshotsEditorSettings::Get()
 {
 	return GetMutableDefault<ULevelSnapshotsEditorSettings>();
 }
 
-ULevelSnapshotsEditorSettings::ULevelSnapshotsEditorSettings(const FObjectInitializer& ObjectInitializer)
+ULevelSnapshotsEditorSettings::ULevelSnapshotsEditorSettings()
 {
-	RootLevelSnapshotSaveDir.Path = "/Game/LevelSnapshots";
-	LevelSnapshotSaveDir = "{map}/{year}-{month}-{day}";
-	DefaultLevelSnapshotName = "{map}_{user}_{time}";
-
-	bEnableLevelSnapshotsToolbarButton = true;
-	bUseCreationForm = true;
-	
-	const FVector2D DefaultClientSize = FVector2D(400.f, 400.f);
-
 	float DPIScale = 1.0f;
-
 	if (FSlateApplicationBase::IsInitialized())
 	{
 		const FSlateRect WorkAreaRect = FSlateApplicationBase::Get().GetPreferredWorkArea();
 		DPIScale = FPlatformApplicationMisc::GetDPIScaleFactorAtPoint(WorkAreaRect.Left, WorkAreaRect.Top);
 	}
 
+	const FVector2D DefaultClientSize = FVector2D(400.f, 400.f);
 	PreferredCreationFormWindowWidth = DefaultClientSize.X * DPIScale;
 	PreferredCreationFormWindowHeight = DefaultClientSize.Y * DPIScale;
 }
@@ -50,25 +36,6 @@ void ULevelSnapshotsEditorSettings::SetLastCreationWindowSize(const FVector2D In
 	PreferredCreationFormWindowHeight = InLastSize.Y;
 }
 
-const FString& ULevelSnapshotsEditorSettings::GetNameOverride() const
-{
-	return LevelSnapshotNameOverride.IsSet() ? LevelSnapshotNameOverride.Get() : DefaultLevelSnapshotName;
-}
-
-void ULevelSnapshotsEditorSettings::SetNameOverride(const FString& InName)
-{
-	LevelSnapshotNameOverride = InName;
-}
-
-void ULevelSnapshotsEditorSettings::ValidateRootLevelSnapshotSaveDirAsGameContentRelative()
-{
-	// Enforce Game Content Dir
-	if (!RootLevelSnapshotSaveDir.Path.StartsWith("/Game/"))
-	{
-		RootLevelSnapshotSaveDir.Path = "/Game/";
-	}
-}
-
 void ULevelSnapshotsEditorSettings::SanitizePathInline(FString& InPath, const bool bSkipForwardSlash)
 {
 	FString IllegalChars = FPaths::GetInvalidFileSystemChars().ReplaceEscapedCharWithChar() + " .";
@@ -89,41 +56,56 @@ void ULevelSnapshotsEditorSettings::SanitizePathInline(FString& InPath, const bo
 
 void ULevelSnapshotsEditorSettings::SanitizeAllProjectSettingsPaths(const bool bSkipForwardSlash)
 {
+	FString Dummy;
+	const FString PackagePath = FPaths::Combine(RootLevelSnapshotSaveDir.Path, FGuid::NewGuid().ToString());
+	// If you enter a path that's not in /Game/ nor in any plugin, this returns false.
+	if (!FPackageName::TryConvertLongPackageNameToFilename(PackagePath, Dummy, FPackageName::GetAssetPackageExtension()))
+	{
+		RootLevelSnapshotSaveDir.Path = TEXT("/Game/LevelSnapshots"); 
+	}
+	
 	SanitizePathInline(RootLevelSnapshotSaveDir.Path, bSkipForwardSlash);
 	SanitizePathInline(LevelSnapshotSaveDir, bSkipForwardSlash);
 	SanitizePathInline(DefaultLevelSnapshotName, bSkipForwardSlash);
 }
 
-FFormatNamedArguments ULevelSnapshotsEditorSettings::GetFormatNamedArguments(const FString& InWorldName)
+namespace UE::LevelSnapshots::Editor
 {
-	FNumberFormattingOptions IntOptions;
-	IntOptions.MinimumIntegralDigits = 2;
-
-	const FDateTime& LocalNow = FDateTime::Now();
-
-	FFormatNamedArguments FormatArguments;
-	FormatArguments.Add("map", FText::FromString(InWorldName));
-	FormatArguments.Add("user", FText::FromString(FPlatformProcess::UserName()));
-	FormatArguments.Add("year", FText::FromString(FString::FromInt(LocalNow.GetYear())));
-	FormatArguments.Add("month", FText::AsNumber(LocalNow.GetMonth(), &IntOptions));
-	FormatArguments.Add("day", FText::AsNumber(LocalNow.GetDay(), &IntOptions));
-	FormatArguments.Add("date", FText::Format(FText::FromString("{0}-{1}-{2}"), FormatArguments["year"], FormatArguments["month"], FormatArguments["day"]));
-	FormatArguments.Add("time",
-		FText::Format(
-			FText::FromString("{0}-{1}-{2}"),
-			FText::AsNumber(LocalNow.GetHour(), &IntOptions), FText::AsNumber(LocalNow.GetMinute(), &IntOptions), FText::AsNumber(LocalNow.GetSecond(), &IntOptions)));
-
-	return FormatArguments;
+	static FFormatNamedArguments GetFormatNamedArguments(const FString& InWorldName)
+	{
+		FNumberFormattingOptions IntOptions;
+		IntOptions.MinimumIntegralDigits = 2;
+
+		const FDateTime& LocalNow = FDateTime::Now();
+
+		FFormatNamedArguments FormatArguments;
+		FormatArguments.Add("map", FText::FromString(InWorldName));
+		FormatArguments.Add("user", FText::FromString(FPlatformProcess::UserName()));
+		FormatArguments.Add("year", FText::FromString(FString::FromInt(LocalNow.GetYear())));
+		FormatArguments.Add("month", FText::AsNumber(LocalNow.GetMonth(), &IntOptions));
+		FormatArguments.Add("day", FText::AsNumber(LocalNow.GetDay(), &IntOptions));
+		FormatArguments.Add("date", FText::Format(FText::FromString("{0}-{1}-{2}"), FormatArguments["year"], FormatArguments["month"], FormatArguments["day"]));
+		FormatArguments.Add("time",
+			FText::Format(
+				FText::FromString("{0}-{1}-{2}"),
+				FText::AsNumber(LocalNow.GetHour(), &IntOptions), FText::AsNumber(LocalNow.GetMinute(), &IntOptions), FText::AsNumber(LocalNow.GetSecond(), &IntOptions)));
+
+		return FormatArguments;
+	}
 }
 
 FText ULevelSnapshotsEditorSettings::ParseLevelSnapshotsTokensInText(const FText& InTextToParse, const FString& InWorldName)
 {
-	const FFormatNamedArguments& FormatArguments = GetFormatNamedArguments(InWorldName);
-
+	const FFormatNamedArguments FormatArguments = UE::LevelSnapshots::Editor::GetFormatNamedArguments(InWorldName);
 	return FText::Format(InTextToParse, FormatArguments);
 }
 
-bool ULevelSnapshotsEditorSettings::IsNameOverridden() const
+#if WITH_EDITOR
+void ULevelSnapshotsEditorSettings::PostEditChangeProperty(FPropertyChangedEvent& PropertyChangedEvent)
 {
-	return LevelSnapshotNameOverride.IsSet();
+	UObject::PostEditChangeProperty(PropertyChangedEvent);
+
+	SanitizeAllProjectSettingsPaths(true);
+	SaveConfig();
 }
+#endif
diff --git a/Engine/Plugins/VirtualProduction/LevelSnapshots/Source/LevelSnapshotsEditor/Private/LevelSnapshotsEditorSettings.h b/Engine/Plugins/VirtualProduction/LevelSnapshots/Source/LevelSnapshotsEditor/Private/LevelSnapshotsEditorSettings.h
--- a/Engine/Plugins/VirtualProduction/LevelSnapshots/Source/LevelSnapshotsEditor/Private/LevelSnapshotsEditorSettings.h
+++ b/Engine/Plugins/VirtualProduction/LevelSnapshots/Source/LevelSnapshotsEditor/Private/LevelSnapshotsEditorSettings.h
@@ -17,35 +17,28 @@ class LEVELSNAPSHOTSEDITOR_API ULevelSnapshotsEditorSettings : public UObject
 	GENERATED_BODY()
 public:
 
-	static const ULevelSnapshotsEditorSettings* Get();
-	static ULevelSnapshotsEditorSettings* GetMutable();
-	
-	ULevelSnapshotsEditorSettings(const FObjectInitializer& ObjectInitializer);
+	static ULevelSnapshotsEditorSettings* Get();
+	ULevelSnapshotsEditorSettings();
 
 	FVector2D GetLastCreationWindowSize() const;
 	/* Setting the Window Size through code will not save the size to the config. To make sure it's saved, call SaveConfig(). */
 	void SetLastCreationWindowSize(const FVector2D InLastSize);
-
-	const FString& GetNameOverride() const;
-	void SetNameOverride(const FString& InName);
-	
-	void ValidateRootLevelSnapshotSaveDirAsGameContentRelative();
 	
-	static void SanitizePathInline(FString& InPath, const bool bSkipForwardSlash);
-
-	/* Removes /?:&\*"<>|%#@^ . from project settings path strings. Optionally the forward slash can be kept so that the user can define a file structure. */
+	/** Removes forbidden file path characters (e.g. /?:&\*"<>|%#@^. ) from project settings path strings. Optionally the forward slash can be kept so that the user can define a file structure. */
 	void SanitizeAllProjectSettingsPaths(const bool bSkipForwardSlash);
-	
-	static FFormatNamedArguments GetFormatNamedArguments(const FString& InWorldName);
+	/** Removes forbidden file path characters (e.g. /?:&\*"<>|%#@^ ) */
+	static void SanitizePathInline(FString& InPath, const bool bSkipForwardSlash);
 
 	UFUNCTION(BlueprintCallable, Category = "Level Snapshots")
 	static FText ParseLevelSnapshotsTokensInText(const FText& InTextToParse, const FString& InWorldName);
 
-	bool IsNameOverridden() const;
-
-	// Must be a directory in the Game Content folder ("/Game/"). For best results, use the picker.  
-	UPROPERTY(config, EditAnywhere, Category = "Data", meta = (RelativeToGameContentDir, ContentDir))
-	FDirectoryPath RootLevelSnapshotSaveDir;
+#if WITH_EDITOR
+	virtual void PostEditChangeProperty(FPropertyChangedEvent& PropertyChangedEvent) override;
+#endif
+	
+	/** The base directory in which all snapshots will be saved. LevelSnapshotSaveDir specifies in which subdirectory snapshots are saved. */
+	UPROPERTY(config, EditAnywhere, Category = "Data", meta = (ContentDir))
+	FDirectoryPath RootLevelSnapshotSaveDir { TEXT("/Game/LevelSnapshots") };
 
 	/** The format to use for the resulting filename. Extension will be added automatically. Any tokens of the form {token} will be replaced with the corresponding value:
 	 * {map}		- The name of the captured map or level
@@ -57,7 +50,7 @@ public:
 	 * {time}       - The current time from the local computer in the format of hours-minutes-seconds
 	 */
 	UPROPERTY(config, EditAnywhere, Category = "Data")
-	FString LevelSnapshotSaveDir;
+	FString LevelSnapshotSaveDir = TEXT("{map}/{year}-{month}-{day}");
 
 	/** The format to use for the resulting filename. Extension will be added automatically. Any tokens of the form {token} will be replaced with the corresponding value:
 	 * {map}		- The name of the captured map or level
@@ -69,14 +62,13 @@ public:
 	 * {time}       - The current time from the local computer in the format of hours-minutes-seconds
 	 */
 	UPROPERTY(config, EditAnywhere, Category = "Data")
-	FString DefaultLevelSnapshotName;
-	
+	FString DefaultLevelSnapshotName = TEXT("{map}_{user}_{time}");
 	
 	UPROPERTY(config, EditAnywhere, Category = "Editor", meta = (ConfigRestartRequired = true))
-	bool bEnableLevelSnapshotsToolbarButton;
+	bool bEnableLevelSnapshotsToolbarButton = true;
 
 	UPROPERTY(config, EditAnywhere, Category = "Editor")
-	bool bUseCreationForm;
+	bool bUseCreationForm = true;
 
 	/* If true, clicking on an actor group under 'Modified Actors' will select the actor in the scene. The previous selection will be deselected. */
 	UPROPERTY(config, EditAnywhere, Category = "Editor")
@@ -87,9 +79,4 @@ public:
 
 	UPROPERTY(config, EditAnywhere, Category = "Editor")
 	float PreferredCreationFormWindowHeight;
-
-private:
-	
-	/* If the user overrides the Name field in the creation form, the override will be saved here so it can be recalled. */
-	TAttribute<FString> LevelSnapshotNameOverride;
 };
diff --git a/Engine/Plugins/VirtualProduction/LevelSnapshots/Source/LevelSnapshotsEditor/Private/Util/TakeSnapshotUtil.cpp b/Engine/Plugins/VirtualProduction/LevelSnapshots/Source/LevelSnapshotsEditor/Private/Util/TakeSnapshotUtil.cpp
--- a/Engine/Plugins/VirtualProduction/LevelSnapshots/Source/LevelSnapshotsEditor/Private/Util/TakeSnapshotUtil.cpp
+++ b/Engine/Plugins/VirtualProduction/LevelSnapshots/Source/LevelSnapshotsEditor/Private/Util/TakeSnapshotUtil.cpp
@@ -31,7 +31,7 @@ namespace UE::LevelSnapshots::Editor
 {
 	TAutoConsoleVariable<bool> CVarSaveAsync(TEXT("LevelSnapshots.SnapshotFormUsesSaveAsync"), true, TEXT("Set whether the Take Snapshot Form, opened by pressing the toolbar icon, saves snapshots using the SAVE_Async flag."));
 	
-	static void HandleFormReply(const FText& InDescription, bool bShouldUseOverrides)
+	static void HandleFormReply(const FText& InDescription, const FString& SnapshotName)
 	{
 		UWorld* World = ULevelSnapshotsEditorData::GetEditorWorld();
 		if (!ensure(World))
@@ -39,22 +39,19 @@ namespace UE::LevelSnapshots::Editor
 			return;
 		}
 		
-		ULevelSnapshotsEditorSettings* DataManagementSettings = ULevelSnapshotsEditorSettings::GetMutable();
-		DataManagementSettings->ValidateRootLevelSnapshotSaveDirAsGameContentRelative();
+		ULevelSnapshotsEditorSettings* DataManagementSettings = ULevelSnapshotsEditorSettings::Get();
 		DataManagementSettings->SanitizeAllProjectSettingsPaths(true);
 
-		const FText& NewSnapshotDir = ULevelSnapshotsEditorSettings::ParseLevelSnapshotsTokensInText(
+		const FText NewSnapshotDir = ULevelSnapshotsEditorSettings::ParseLevelSnapshotsTokensInText(
 			FText::FromString(FPaths::Combine(DataManagementSettings->RootLevelSnapshotSaveDir.Path, DataManagementSettings->LevelSnapshotSaveDir)),
 			World->GetName()
 			);
-		const FString& DescriptionString = bShouldUseOverrides && DataManagementSettings->IsNameOverridden() ?	DataManagementSettings->GetNameOverride() : DataManagementSettings->DefaultLevelSnapshotName;
-		const FText& NewSnapshotName = ULevelSnapshotsEditorSettings::ParseLevelSnapshotsTokensInText(
-				FText::FromString(DescriptionString),
+		const FText NewSnapshotName = ULevelSnapshotsEditorSettings::ParseLevelSnapshotsTokensInText(
+				FText::FromString(SnapshotName),
 				World->GetName()
 				);
 
-		const FString& ValidatedName = FPaths::MakeValidFileName(NewSnapshotName.ToString());
-
+		const FString ValidatedName = FPaths::MakeValidFileName(NewSnapshotName.ToString());
 		const bool bSaveAsync = CVarSaveAsync.GetValueOnAnyThread();
 		SnapshotEditor::TakeLevelSnapshotAndSaveToDisk(World, ValidatedName, NewSnapshotDir.ToString(), InDescription.ToString(), false, bSaveAsync);
 	}
@@ -70,55 +67,46 @@ namespace UE::LevelSnapshots::Editor
 
 void SnapshotEditor::TakeSnapshotWithOptionalForm()
 {
-	FLevelSnapshotsEditorModule& Module = FLevelSnapshotsEditorModule::Get();
 	if (ULevelSnapshotsEditorSettings::Get()->bUseCreationForm)
 	{
-		TSharedRef<SWidget> CreationForm = SLevelSnapshotsEditorCreationForm::MakeAndShowCreationWindow(
-			FCloseCreationFormDelegate::CreateLambda([](const FText& Description)
-			{
-				UE::LevelSnapshots::Editor::HandleFormReply(Description, true);
-			}));
+		SLevelSnapshotsEditorCreationForm::MakeAndShowCreationWindow(FCloseCreationFormDelegate::CreateStatic(UE::LevelSnapshots::Editor::HandleFormReply));
 	}
 	else
 	{
-		UE::LevelSnapshots::Editor::HandleFormReply(FText::GetEmpty(), false);
+		UE::LevelSnapshots::Editor::HandleFormReply(FText::GetEmpty(), ULevelSnapshotsEditorSettings::Get()->DefaultLevelSnapshotName);
 	}
 }
 
 ULevelSnapshot* SnapshotEditor::TakeLevelSnapshotAndSaveToDisk(UWorld* World, const FString& FileName, const FString& FolderPath, const FString& Description, bool bShouldCreateUniqueFileName, bool bSaveAsync)
 {
 	SCOPED_SNAPSHOT_EDITOR_TRACE(CreateSnapshotPackage);
+	
 	FString AssetName = FileName;
-    	
-    FString BasePackageName = FPaths::Combine(FolderPath, FileName);
-
-    BasePackageName.RemoveFromStart(TEXT("/"));
-    BasePackageName.RemoveFromStart(TEXT("Content/"));
-    BasePackageName.StartsWith(TEXT("Game/")) == true ? BasePackageName.InsertAt(0, TEXT("/")) : BasePackageName.InsertAt(0, TEXT("/Game/"));
-
-    FString PackageName = BasePackageName;
-
+    FString PackageName = FPaths::Combine(FolderPath, AssetName);
     if (bShouldCreateUniqueFileName)
     {
     	IAssetTools& AssetTools = FModuleManager::Get().LoadModuleChecked<FAssetToolsModule>("AssetTools").Get();
-    	
-    	AssetTools.CreateUniqueAssetName(BasePackageName, TEXT(""), PackageName, AssetName);
+    	AssetTools.CreateUniqueAssetName(PackageName, TEXT(""), PackageName, AssetName);
     }
 
-    const FString& PackageFileName = FPackageName::LongPackageNameToFilename(PackageName, FPackageName::GetAssetPackageExtension());
-    const bool bPathIsValid = FPaths::ValidatePath(PackageFileName);
-
-    bool bProceedWithSave = bPathIsValid;
-
+	FString PackageFileName;
+	FPackageName::TryConvertLongPackageNameToFilename(PackageName, PackageFileName, FPackageName::GetAssetPackageExtension());
+	
+    bool bProceedWithSave = FPaths::ValidatePath(PackageFileName);
     if (bProceedWithSave && !bShouldCreateUniqueFileName)
     {
     	FAssetRegistryModule& AssetRegistryModule = FModuleManager::LoadModuleChecked<FAssetRegistryModule>("AssetRegistry");
-    	FAssetData ExistingAsset = AssetRegistryModule.Get().GetAssetByObjectPath(FName(*(BasePackageName + "." + FileName)));
-    	if ( ExistingAsset.IsValid() && ExistingAsset.AssetClassPath == ULevelSnapshot::StaticClass()->GetClassPathName() )
+    	FAssetData ExistingAsset = AssetRegistryModule.Get().GetAssetByObjectPath(FName(*(PackageName + "." + FileName)));
+    	if (ExistingAsset.IsValid() && ExistingAsset.AssetClassPath == ULevelSnapshot::StaticClass()->GetClassPathName())
     	{
     		EAppReturnType::Type ShouldReplace = FMessageDialog::Open( EAppMsgType::YesNo, FText::Format(LOCTEXT("ReplaceAssetMessage", "{0} already exists. Do you want to replace it?"), FText::FromString(FileName)) );
     		bProceedWithSave = (ShouldReplace == EAppReturnType::Yes);
     	}
+    	else if (ExistingAsset.IsValid())
+    	{
+    		UE_LOG(LogLevelSnapshots, Error, TEXT("Asset %s already exists and is not a snapshot. Aborting."))
+    		return nullptr;
+    	}
     }
 
     if (bProceedWithSave)
@@ -141,17 +129,18 @@ ULevelSnapshot* SnapshotEditor::TakeLevelSnapshotAndSaveToDisk(UWorld* World, co
     	UPackage* SavePackage = CreatePackage(*PackageName);
     	const EObjectFlags AssetFlags = RF_Public | RF_Standalone;
     	ULevelSnapshot* SnapshotAsset = NewObject<ULevelSnapshot>(SavePackage, *AssetName, AssetFlags, nullptr);
-
-    	check(SavePackage && SnapshotAsset);
+    	if (!ensure(SavePackage && SnapshotAsset))
+    	{
+    		return nullptr;
+    	}
     	
     	SnapshotAsset->SetSnapshotName(*FileName);
     	SnapshotAsset->SetSnapshotDescription(Description);
-
         const bool bSuccessful = SnapshotAsset->SnapshotWorld(World);
         if (!bSuccessful)
         {
             // Package and snapshot need to be destroyed again
-            UE::LevelSnapshots::Editor:: DestroySnapshot(SnapshotAsset);
+            UE::LevelSnapshots::Editor::DestroySnapshot(SnapshotAsset);
             ObjectTools::DeleteObjectsUnchecked({ SavePackage });
             
             NotificationItem->SetText(
@@ -160,8 +149,8 @@ ULevelSnapshot* SnapshotEditor::TakeLevelSnapshotAndSaveToDisk(UWorld* World, co
             NotificationItem->SetCompletionState(SNotificationItem::CS_Fail);
             return nullptr;
         }
+    	
     	SnapshotAsset->MarkPackageDirty();
-
     	FAssetRegistryModule::AssetCreated(SnapshotAsset);
     	GenerateThumbnailForSnapshotAsset(SnapshotAsset);
 
@@ -177,7 +166,7 @@ ULevelSnapshot* SnapshotEditor::TakeLevelSnapshotAndSaveToDisk(UWorld* World, co
     	}
     	
     	// Notify the user of the outcome
-    	if (bPathIsValid && bSavingSuccessful)
+    	if (bSavingSuccessful)
     	{
     		NotificationItem->SetText(
     			FText::Format(
diff --git a/Engine/Plugins/VirtualProduction/LevelSnapshots/Source/LevelSnapshotsEditor/Private/Widgets/SLevelSnapshotsEditorCreationForm.cpp b/Engine/Plugins/VirtualProduction/LevelSnapshots/Source/LevelSnapshotsEditor/Private/Widgets/SLevelSnapshotsEditorCreationForm.cpp
--- a/Engine/Plugins/VirtualProduction/LevelSnapshots/Source/LevelSnapshotsEditor/Private/Widgets/SLevelSnapshotsEditorCreationForm.cpp
+++ b/Engine/Plugins/VirtualProduction/LevelSnapshots/Source/LevelSnapshotsEditor/Private/Widgets/SLevelSnapshotsEditorCreationForm.cpp
@@ -87,7 +87,6 @@ void SLevelSnapshotsEditorCreationForm::Construct(
 {
 	WidgetWindow = InWidgetWindow;
 	CallOnCloseDelegate = CallOnClose;
-	bNameDiffersFromDefault = ULevelSnapshotsEditorSettings::Get()->IsNameOverridden();
 	
 	FSlateFontInfo NameTextFont = FAppStyle::Get().GetFontStyle("Regular");
 	NameTextFont.Size = 9;
@@ -240,7 +239,7 @@ TSharedRef<SWidget> SLevelSnapshotsEditorCreationForm::MakeDataManagementSetting
 	DetailsViewArgs.bShowScrollBar = false;
 
 	TSharedRef<IDetailsView> Details = PropertyEditorModule.CreateDetailView(DetailsViewArgs);
-	const TArray<UObject*> ProjectSettingsObjects = { ULevelSnapshotsEditorSettings::GetMutable() };
+	const TArray<UObject*> ProjectSettingsObjects = { ULevelSnapshotsEditorSettings::Get() };
 	// By requirement, we're only supposed to show the data management settings
 	Details->RegisterInstancedCustomPropertyLayout(
 		ULevelSnapshotsEditorSettings::StaticClass(),
@@ -252,24 +251,20 @@ TSharedRef<SWidget> SLevelSnapshotsEditorCreationForm::MakeDataManagementSetting
 	return Details;
 }
 
-SLevelSnapshotsEditorCreationForm::~SLevelSnapshotsEditorCreationForm()
-{
-	CallOnCloseDelegate.Unbind();
-}
-
 FText SLevelSnapshotsEditorCreationForm::GetNameOverrideText() const
 {
 	const ULevelSnapshotsEditorSettings* Settings = ULevelSnapshotsEditorSettings::Get();
 	check(Settings);
-	
+
+	const FString SnapshotName = NameOverride.Get(Settings->DefaultLevelSnapshotName);
 	UWorld* World = ULevelSnapshotsEditorData::GetEditorWorld();
 	if (!ensure(World))
 	{
-		return FText::FromString(Settings->GetNameOverride());
+		return FText::FromString(SnapshotName);
 	}
 
 	return ULevelSnapshotsEditorSettings::ParseLevelSnapshotsTokensInText(
-		FText::FromString(Settings->GetNameOverride()),
+		FText::FromString(SnapshotName),
 		World->GetName()
 		);
 }
@@ -278,10 +273,7 @@ void SLevelSnapshotsEditorCreationForm::SetNameOverrideText(const FText& InNewTe
 {
 	FString NameAsString = InNewText.ToString();
 	ULevelSnapshotsEditorSettings::SanitizePathInline(NameAsString, true);
-
-	ULevelSnapshotsEditorSettings* Settings = ULevelSnapshotsEditorSettings::GetMutable();
-	Settings->SetNameOverride(NameAsString);
-	bNameDiffersFromDefault = Settings->IsNameOverridden();
+	NameOverride = NameAsString;
 }
 
 void SLevelSnapshotsEditorCreationForm::SetDescriptionText(const FText& InNewText, ETextCommit::Type InCommitType)
@@ -291,7 +283,7 @@ void SLevelSnapshotsEditorCreationForm::SetDescriptionText(const FText& InNewTex
 
 EVisibility SLevelSnapshotsEditorCreationForm::GetNameDiffersFromDefaultAsVisibility() const
 {
-	return bNameDiffersFromDefault ? EVisibility::Visible : EVisibility::Hidden;
+	return NameOverride.IsSet() ? EVisibility::Visible : EVisibility::Hidden;
 }
 
 FReply SLevelSnapshotsEditorCreationForm::OnResetNameClicked()
@@ -302,7 +294,7 @@ FReply SLevelSnapshotsEditorCreationForm::OnResetNameClicked()
 
 FReply SLevelSnapshotsEditorCreationForm::OnCreateButtonPressed()
 {
-	bWasCreateSnapshotPressed = true;
+	CallOnCloseDelegate.ExecuteIfBound(DescriptionText, GetNameOverrideText().ToString());
 
 	check(WidgetWindow.IsValid());
 	WidgetWindow.Pin()->RequestDestroyWindow();
@@ -313,12 +305,7 @@ FReply SLevelSnapshotsEditorCreationForm::OnCreateButtonPressed()
 void SLevelSnapshotsEditorCreationForm::OnWindowClosed(const TSharedRef<SWindow>& ParentWindow) const
 {
 	const FVector2D WindowSize = ParentWindow->GetClientSizeInScreen();
-	ULevelSnapshotsEditorSettings* Settings = ULevelSnapshotsEditorSettings::GetMutable();
+	ULevelSnapshotsEditorSettings* Settings = ULevelSnapshotsEditorSettings::Get();
 	Settings->SetLastCreationWindowSize(WindowSize);
 	Settings->SaveConfig();
-
-	if (bWasCreateSnapshotPressed)
-	{
-		CallOnCloseDelegate.ExecuteIfBound(DescriptionText);
-	}
 }
diff --git a/Engine/Plugins/VirtualProduction/LevelSnapshots/Source/LevelSnapshotsEditor/Private/Widgets/SLevelSnapshotsEditorCreationForm.h b/Engine/Plugins/VirtualProduction/LevelSnapshots/Source/LevelSnapshotsEditor/Private/Widgets/SLevelSnapshotsEditorCreationForm.h
--- a/Engine/Plugins/VirtualProduction/LevelSnapshots/Source/LevelSnapshotsEditor/Private/Widgets/SLevelSnapshotsEditorCreationForm.h
+++ b/Engine/Plugins/VirtualProduction/LevelSnapshots/Source/LevelSnapshotsEditor/Private/Widgets/SLevelSnapshotsEditorCreationForm.h
@@ -5,13 +5,13 @@
 #include "CoreMinimal.h"
 #include "Dialog/SCustomDialog.h"
 
-DECLARE_DELEGATE_OneParam(FCloseCreationFormDelegate, const FText& /* Description */);
+DECLARE_DELEGATE_TwoParams(FCloseCreationFormDelegate, const FText& /*Description*/, const FString& /*SnapshotName*/);
 
 class ULevelSnapshotsSettings;
 class ULevelSnapshotsEditorSettings;
 class SWindow;
 
-class SLevelSnapshotsEditorCreationForm : public SCustomDialog
+class SLevelSnapshotsEditorCreationForm : public SCompoundWidget
 {
 public:
 
@@ -21,18 +21,9 @@ public:
 	{}
 	SLATE_END_ARGS()
 
-	void Construct(
-		const FArguments& InArgs,
-		TWeakPtr<SWindow> InWidgetWindow,
-		const FCloseCreationFormDelegate& CallOnClose
-		);
-
-	~SLevelSnapshotsEditorCreationForm();
-
-	TSharedRef<SWidget> MakeDataManagementSettingsDetailsWidget() const;
-
+	void Construct(const FArguments& InArgs, TWeakPtr<SWindow> InWidgetWindow, const FCloseCreationFormDelegate& CallOnClose);
+	
 	FText GetNameOverrideText() const;
-
 	void SetNameOverrideText(const FText& InNewText, ETextCommit::Type InCommitType);
 	void SetDescriptionText(const FText& InNewText, ETextCommit::Type InCommitType);
 
@@ -41,17 +32,14 @@ public:
 	FReply OnResetNameClicked();
 	FReply OnCreateButtonPressed();
 
-	void OnWindowClosed(const TSharedRef<SWindow>& ParentWindow) const;
-
 private:
 	
-	TWeakPtr< SWindow > WidgetWindow;
-	TSharedPtr<SWidget> ResetPathButton;
-
-	bool bNameDiffersFromDefault = false;
-	bool bWasCreateSnapshotPressed = false;
+	TWeakPtr<SWindow> WidgetWindow;
 
 	FText DescriptionText;
-
+	TOptional<FString> NameOverride;
 	FCloseCreationFormDelegate CallOnCloseDelegate;
+	
+	TSharedRef<SWidget> MakeDataManagementSettingsDetailsWidget() const;
+	void OnWindowClosed(const TSharedRef<SWindow>& ParentWindow) const;
 };
