- Fixing crash when invalid path is provided when saving a snapshot

diff --git a/Engine/Plugins/VirtualProduction/LevelSnapshots/Source/LevelSnapshotsEditor/Private/Util/TakeSnapshotUtil.cpp b/Engine/Plugins/VirtualProduction/LevelSnapshots/Source/LevelSnapshotsEditor/Private/Util/TakeSnapshotUtil.cpp
--- a/Engine/Plugins/VirtualProduction/LevelSnapshots/Source/LevelSnapshotsEditor/Private/Util/TakeSnapshotUtil.cpp
+++ b/Engine/Plugins/VirtualProduction/LevelSnapshots/Source/LevelSnapshotsEditor/Private/Util/TakeSnapshotUtil.cpp
@@ -105,10 +105,17 @@ ULevelSnapshot* SnapshotEditor::TakeLevelSnapshotAndSaveToDisk(UWorld* World, co
     	else if (ExistingAsset.IsValid())
     	{
     		UE_LOG(LogLevelSnapshots, Error, TEXT("Asset %s already exists and is not a snapshot. Aborting."))
-    		return nullptr;
+			bProceedWithSave = false;
     	}
     }
 
+	FText ErrorReason;
+	if (!FPackageName::IsValidLongPackageName(PackageName, true, &ErrorReason))
+	{
+		UE_LOG(LogLevelSnapshots, Error, TEXT("Invalid package name. Error: %s"), *ErrorReason.ToString());
+		bProceedWithSave = false;
+	}
+
     if (bProceedWithSave)
     {
     	// Show notification
