Fix 2 crashes
- Failing to load classes causes a crash: Added nullptr checks in all locations that make use of TryLoadClass as well
- Name clashses archetype objects. Snapshots used the global GetTransientPackage() so if you loaded two snapshots, you'd get a crash. Snaphosts now create archetypes within the same package that contains the transient snapshot world.

diff --git a/Engine/Plugins/VirtualProduction/LevelSnapshots/Source/LevelSnapshots/Private/Data/Util/Component/BaseComponentRestorer.h b/Engine/Plugins/VirtualProduction/LevelSnapshots/Source/LevelSnapshots/Private/Data/Util/Component/BaseComponentRestorer.h
--- a/Engine/Plugins/VirtualProduction/LevelSnapshots/Source/LevelSnapshots/Private/Data/Util/Component/BaseComponentRestorer.h
+++ b/Engine/Plugins/VirtualProduction/LevelSnapshots/Source/LevelSnapshots/Private/Data/Util/Component/BaseComponentRestorer.h
@@ -91,9 +91,11 @@ namespace UE::LevelSnapshots::Private
 				if (ComponentOuter)
 				{
 					Component = AllocateComponent(ComponentPath, SubobjectData, ComponentData, ComponentOuter);
-			
-					// UActorComponent::PostInitProperties implicitly calls AddOwnedComponent but we have to manually add it to the other arrays
-					AddToCorrectComponentArray(ComponentData, Component);
+					if (Component)
+					{
+						// UActorComponent::PostInitProperties implicitly calls AddOwnedComponent but we have to manually add it to the other arrays
+						AddToCorrectComponentArray(ComponentData, Component);
+					}
 				}
 			}
 		
@@ -106,6 +108,12 @@ namespace UE::LevelSnapshots::Private
 		{
 			const FSoftClassPath ClassPath = GetClass(SubobjectData, WorldData);
 			UClass* ComponentClass = ClassPath.TryLoadClass<UActorComponent>();
+			if (!ComponentClass)
+			{
+				UE_LOG(LogLevelSnapshots, Error, TEXT("Failed to load class %s"), *ClassPath.ToString())
+				return nullptr;
+			}
+			
 			const FName ComponentName = *UE::LevelSnapshots::Private::ExtractLastSubobjectName(ComponentPath);
 			const EObjectFlags ObjectFlags = SubobjectData.GetObjectFlags();
 			if constexpr (TDerived::IsRestoringIntoSnapshotWorld())
diff --git a/Engine/Plugins/VirtualProduction/LevelSnapshots/Source/LevelSnapshots/Private/Data/Util/WorldData/ClassDataUtil.cpp b/Engine/Plugins/VirtualProduction/LevelSnapshots/Source/LevelSnapshots/Private/Data/Util/WorldData/ClassDataUtil.cpp
--- a/Engine/Plugins/VirtualProduction/LevelSnapshots/Source/LevelSnapshots/Private/Data/Util/WorldData/ClassDataUtil.cpp
+++ b/Engine/Plugins/VirtualProduction/LevelSnapshots/Source/LevelSnapshots/Private/Data/Util/WorldData/ClassDataUtil.cpp
@@ -44,7 +44,7 @@ namespace UE::LevelSnapshots::Private
 		}
 		
 		UObject* CDO = NewObject<UObject>(
-			GetTransientPackage(),
+			WorldData.SnapshotWorld->GetPackage(),
 			Class,
 			*FString("SnapshotCDO_").Append(*MakeUniqueObjectName(GetTransientPackage(), Class).ToString())
 			// TODO: Set RF_Archetype
@@ -69,6 +69,11 @@ namespace UE::LevelSnapshots::Private
 		}
 
 		UClass* Class = ArchetypeData->ClassPath.TryLoadClass<UObject>();
+		if (!Class)
+		{
+			return {};
+		}
+		
 		UObject* Archetype = NewObject<UObject>(
 			GetTransientPackage(),
 			Class,
@@ -106,8 +111,15 @@ namespace UE::LevelSnapshots::Private
 			return { &FallbackData.GetValue() };
 		}
 
+		UClass* Class = SavedClassData.ClassPath.TryLoadClass<UObject>();
+		if (!Class)
+		{
+			UE_LOG(LogLevelSnapshots, Error, TEXT("Failed to load class %s"), *SavedClassData.ClassPath.ToString());
+			return {};
+		}
+		
 		// 2. First time encountering this class, get the real archetype
-		UObject* RealArchetype = UObject::GetArchetypeFromRequiredInfo(SavedClassData.ClassPath.TryLoadClass<UObject>(), FallbackInfo.SubobjectOuter, FallbackInfo.SubobjectName, FallbackInfo.SubobjectFlags);
+		UObject* RealArchetype = UObject::GetArchetypeFromRequiredInfo(Class, FallbackInfo.SubobjectOuter, FallbackInfo.SubobjectName, FallbackInfo.SubobjectFlags);
 		if (ensure(RealArchetype))
 		{
 			FallbackData = FClassSnapshotData(); // Optional needs to be assigned
