Fix object names not being unique

diff --git a/Engine/Plugins/VirtualProduction/LevelSnapshots/Source/LevelSnapshots/Private/Data/Util/WorldData/ClassDataUtil.cpp b/Engine/Plugins/VirtualProduction/LevelSnapshots/Source/LevelSnapshots/Private/Data/Util/WorldData/ClassDataUtil.cpp
--- a/Engine/Plugins/VirtualProduction/LevelSnapshots/Source/LevelSnapshots/Private/Data/Util/WorldData/ClassDataUtil.cpp
+++ b/Engine/Plugins/VirtualProduction/LevelSnapshots/Source/LevelSnapshots/Private/Data/Util/WorldData/ClassDataUtil.cpp
@@ -46,7 +46,7 @@ namespace UE::LevelSnapshots::Private
 		UObject* CDO = NewObject<UObject>(
 			WorldData.SnapshotWorld->GetPackage(),
 			Class,
-			*FString("SnapshotCDO_").Append(*MakeUniqueObjectName(GetTransientPackage(), Class).ToString())
+			MakeUniqueObjectName(WorldData.SnapshotWorld->GetPackage(), Class)
 			// TODO: Set RF_Archetype
 			);
 		FApplyClassDefaulDataArchive::SerializeClassDefaultObject(ClassData, WorldData, CDO);
