Level Snapshots should save RootComponent even when it is excluded by FComponentEditorUtils::CanEditComponentInstance.
Core changes are in SnapshotRestorability.h:
- IsComponentDesirableForCapture will now cause RootComponents to be saved even when FComponentEditorUtils::CanEditComponentInstance returns false
- IsComponentRestorable: New function will now check whether 1. IsComponentDesirableForCapture and 2. whether we saved data for the component. If there is no saved data, LevelSnapshots will not attempt to restore that component; previously we did not save data for the RootComponent (for which FComponentEditorUtils::CanEditComponentInstance returned false): this check prevents us from resetting data.

diff --git a/Engine/Plugins/VirtualProduction/LevelSnapshots/Source/LevelSnapshots/Private/Data/Util/EquivalenceUtil.cpp b/Engine/Plugins/VirtualProduction/LevelSnapshots/Source/LevelSnapshots/Private/Data/Util/EquivalenceUtil.cpp
--- a/Engine/Plugins/VirtualProduction/LevelSnapshots/Source/LevelSnapshots/Private/Data/Util/EquivalenceUtil.cpp
+++ b/Engine/Plugins/VirtualProduction/LevelSnapshots/Source/LevelSnapshots/Private/Data/Util/EquivalenceUtil.cpp
@@ -91,10 +91,10 @@ namespace UE::LevelSnapshots::Private::Internal
 	}
 
 	/** @return False if one component could not be matched */
-	static bool EnqueueMatchingComponents(TInlineComponentArray<TPair<UObject*, UObject*>>& SnapshotOriginalPairsToProcess, const FWorldSnapshotData& WorldData, AActor* SnapshotActor, AActor* WorldActor)
+	static bool EnqueueMatchingComponents(ULevelSnapshot* Snapshot, TInlineComponentArray<TPair<UObject*, UObject*>>& SnapshotOriginalPairsToProcess, const FWorldSnapshotData& WorldData, AActor* SnapshotActor, AActor* WorldActor)
 	{
 		bool bFoundUnmatchedObjects = false;
-		UE::LevelSnapshots::Private::IterateComponents(SnapshotActor, WorldActor,
+		UE::LevelSnapshots::Private::IterateRestorableComponents(Snapshot, SnapshotActor, WorldActor,
 			[&SnapshotOriginalPairsToProcess, &WorldData, &bFoundUnmatchedObjects](UActorComponent* SnapshotComp, UActorComponent* WorldComp)
 			{
 				bFoundUnmatchedObjects |= EnqueueMatchingCustomSubobjects(SnapshotOriginalPairsToProcess, WorldData, SnapshotComp, WorldComp);
@@ -113,11 +113,12 @@ namespace UE::LevelSnapshots::Private::Internal
 	}
 }
 
-void UE::LevelSnapshots::Private::IterateComponents(AActor* SnapshotActor, AActor* WorldActor, FHandleMatchedActorComponent OnComponentsMatched, FHandleUnmatchedActorComponent OnSnapshotComponentUnmatched, FHandleUnmatchedActorComponent OnWorldComponentUnmatched)
+void UE::LevelSnapshots::Private::IterateRestorableComponents(ULevelSnapshot* Snapshot, AActor* SnapshotActor, AActor* WorldActor, FHandleMatchedActorComponent OnComponentsMatched, FHandleUnmatchedActorComponent OnSnapshotComponentUnmatched, FHandleUnmatchedActorComponent OnWorldComponentUnmatched)
 {
+	const FSoftObjectPath WorldActorPath = WorldActor;
 	for (UActorComponent* WorldComp : WorldActor->GetComponents())
 	{
-		if (!Restorability::IsComponentDesirableForCapture(WorldComp))
+		if (!Restorability::IsComponentRestorable(Snapshot, WorldActorPath, WorldComp))
 		{
 			continue;
 		}
@@ -134,7 +135,7 @@ void UE::LevelSnapshots::Private::IterateComponents(AActor* SnapshotActor, AActo
 
 	for (UActorComponent* SnapshotComp : SnapshotActor->GetComponents())
 	{
-		if (Restorability::IsComponentDesirableForCapture(SnapshotComp)
+		if (Restorability::IsComponentRestorable(Snapshot, WorldActorPath, SnapshotComp)
 			&& Internal::TryFindMatchingComponent(WorldActor, SnapshotComp) == nullptr)
 		{
 			OnSnapshotComponentUnmatched(SnapshotComp);
@@ -191,7 +192,7 @@ bool UE::LevelSnapshots::Private::HasOriginalChangedPropertiesSinceSnapshotWasTa
 	TInlineComponentArray<TPair<UObject*, UObject*>> SnapshotOriginalPairsToProcess;
 	SnapshotOriginalPairsToProcess.Add(TPair<UObject*, UObject*>(SnapshotActor, WorldActor));
 	
-	const bool bFailedToMatchAllComponentObjects = Internal::EnqueueMatchingComponents(SnapshotOriginalPairsToProcess, Snapshot->GetSerializedData(), SnapshotActor, WorldActor);
+	const bool bFailedToMatchAllComponentObjects = Internal::EnqueueMatchingComponents(Snapshot, SnapshotOriginalPairsToProcess, Snapshot->GetSerializedData(), SnapshotActor, WorldActor);
 	if (bFailedToMatchAllComponentObjects)
 	{
 		return true;
diff --git a/Engine/Plugins/VirtualProduction/LevelSnapshots/Source/LevelSnapshots/Private/Data/Util/EquivalenceUtil.h b/Engine/Plugins/VirtualProduction/LevelSnapshots/Source/LevelSnapshots/Private/Data/Util/EquivalenceUtil.h
--- a/Engine/Plugins/VirtualProduction/LevelSnapshots/Source/LevelSnapshots/Private/Data/Util/EquivalenceUtil.h
+++ b/Engine/Plugins/VirtualProduction/LevelSnapshots/Source/LevelSnapshots/Private/Data/Util/EquivalenceUtil.h
@@ -26,7 +26,7 @@ namespace UE::LevelSnapshots::Private
 	 * @param OnSnapshotComponentUnmatched Called when component exists on snapshot but not on editor world
 	 * @param OnWorldComponentUnmatched Called when component exists in editor world but not in snapshot world
 	 */
-	void IterateComponents(AActor* SnapshotActor, AActor* WorldActor, FHandleMatchedActorComponent OnComponentsMatched, FHandleUnmatchedActorComponent OnSnapshotComponentUnmatched, FHandleUnmatchedActorComponent OnWorldComponentUnmatched);
+	void IterateRestorableComponents(ULevelSnapshot* Snapshot, AActor* SnapshotActor, AActor* WorldActor, FHandleMatchedActorComponent OnComponentsMatched, FHandleUnmatchedActorComponent OnSnapshotComponentUnmatched, FHandleUnmatchedActorComponent OnWorldComponentUnmatched);
 
 	/** Tries to find an actor component by following its full outer path, e.g. /Game/Map.Map:PersistentLevel.SomeActor.SomeParentComp.SomeChildComp will find SomeChildComp with an outer SomeParentComp. */
 	UActorComponent* FindMatchingComponent(AActor* ActorToSearchOn, const FSoftObjectPath& ComponentPath);
diff --git a/Engine/Plugins/VirtualProduction/LevelSnapshots/Source/LevelSnapshots/Private/Data/Util/WorldData/WorldDataUtil.cpp b/Engine/Plugins/VirtualProduction/LevelSnapshots/Source/LevelSnapshots/Private/Data/Util/WorldData/WorldDataUtil.cpp
--- a/Engine/Plugins/VirtualProduction/LevelSnapshots/Source/LevelSnapshots/Private/Data/Util/WorldData/WorldDataUtil.cpp
+++ b/Engine/Plugins/VirtualProduction/LevelSnapshots/Source/LevelSnapshots/Private/Data/Util/WorldData/WorldDataUtil.cpp
@@ -25,6 +25,7 @@
 #include "Misc/ScopedSlowTask.h"
 #include "Templates/NonNullPointer.h"
 #include "UObject/Package.h"
+#include "Util/SnapshotUtil.h"
 #if WITH_EDITOR
 #include "Editor.h"
 #include "Editor/EditorEngine.h"
@@ -258,6 +259,36 @@ void UE::LevelSnapshots::Private::ApplyToWorld(FWorldSnapshotData& WorldData, FS
 #endif
 }
 
+bool UE::LevelSnapshots::Private::HasSavedComponentData(const FWorldSnapshotData& WorldData, const FSoftObjectPath& WorldActorPath, const UActorComponent* EditorOrSnapshotComponent)
+{
+	const FActorSnapshotData* ActorData = WorldData.ActorData.Find(WorldActorPath);
+	if (!ActorData)
+	{
+		return false;
+	}
+
+	const FSoftObjectPath SearchedComponentPath = EditorOrSnapshotComponent;
+	const FString ComponentName = EditorOrSnapshotComponent->GetName();
+	const int32 ComponentDataIndex = [&WorldData, ActorData, ComponentName]() -> int32
+	{
+		// For better performance, instead of iterating through all elements in FWorldSnapshotData::SerializedObjectReferences we can use FActorSnapshotData::ComponentData as hint
+		for (const TPair<int32, FComponentSnapshotData>& ComponentData : ActorData->ComponentData)
+		{
+			if (ensure(WorldData.SerializedObjectReferences.IsValidIndex(ComponentData.Key))
+				&& ExtractLastSubobjectName(WorldData.SerializedObjectReferences[ComponentData.Key]) == ComponentName)
+			{
+				return ComponentData.Key;
+			}
+		}
+		return INDEX_NONE;
+	}();
+
+	const FSubobjectSnapshotData* ComponentData = WorldData.Subobjects.Find(ComponentDataIndex);
+	return ComponentData
+		&& ensureMsgf(ComponentDataIndex != INDEX_NONE, TEXT("Did not expect FWorldSnapshotData::Subobjects to contain INDEX_NONE. Investigate."))
+		&& ComponentData->SerializedData.Num() > 0;
+}
+
 namespace UE::LevelSnapshots::Private::Internal
 {
 	static void PreloadClassesForRestore(FWorldSnapshotData& WorldData, const FPropertySelectionMap& SelectionMap)
diff --git a/Engine/Plugins/VirtualProduction/LevelSnapshots/Source/LevelSnapshots/Private/Data/Util/WorldData/WorldDataUtil.h b/Engine/Plugins/VirtualProduction/LevelSnapshots/Source/LevelSnapshots/Private/Data/Util/WorldData/WorldDataUtil.h
--- a/Engine/Plugins/VirtualProduction/LevelSnapshots/Source/LevelSnapshots/Private/Data/Util/WorldData/WorldDataUtil.h
+++ b/Engine/Plugins/VirtualProduction/LevelSnapshots/Source/LevelSnapshots/Private/Data/Util/WorldData/WorldDataUtil.h
@@ -19,4 +19,7 @@ namespace UE::LevelSnapshots::Private
 	
 	/* Applies the saved properties to WorldActor */
 	void ApplyToWorld(FWorldSnapshotData& WorldData, FSnapshotDataCache& Cache, UWorld* WorldToApplyTo, UPackage* LocalisationSnapshotPackage, const FPropertySelectionMap& PropertiesToSerialize);
+
+	/** Whether the component was saved in the world data */
+	bool HasSavedComponentData(const FWorldSnapshotData& WorldData, const FSoftObjectPath& WorldActorPath, const UActorComponent* EditorOrSnapshotComponent);
 }
diff --git a/Engine/Plugins/VirtualProduction/LevelSnapshots/Source/LevelSnapshots/Private/Restorability/SnapshotRestorability.cpp b/Engine/Plugins/VirtualProduction/LevelSnapshots/Source/LevelSnapshots/Private/Restorability/SnapshotRestorability.cpp
--- a/Engine/Plugins/VirtualProduction/LevelSnapshots/Source/LevelSnapshots/Private/Restorability/SnapshotRestorability.cpp
+++ b/Engine/Plugins/VirtualProduction/LevelSnapshots/Source/LevelSnapshots/Private/Restorability/SnapshotRestorability.cpp
@@ -5,6 +5,7 @@
 #include "Interfaces/IActorSnapshotFilter.h"
 #include "Landscape.h"
 #include "LandscapeGizmoActor.h"
+#include "LevelSnapshot.h"
 #include "LevelSnapshotsLog.h"
 #include "LevelSnapshotsModule.h"
 #include "SnapshotCustomVersion.h"
@@ -20,6 +21,7 @@
 #include "GameFramework/DefaultPhysicsVolume.h"
 #include "GameFramework/WorldSettings.h"
 #include "UObject/UnrealType.h"
+#include "Util/WorldData/WorldDataUtil.h"
 #if WITH_EDITOR
 #include "ActorEditorUtils.h"
 #include "Kismet2/ComponentEditorUtils.h"
@@ -162,12 +164,21 @@ bool UE::LevelSnapshots::Restorability::IsComponentDesirableForCapture(const UAc
 			&& Private::Internal::DoesComponentHaveSupportedClassForCapture(Component)
 			&& !Component->HasAnyFlags(RF_Transient)
 #if WITH_EDITORONLY_DATA
-			&& FComponentEditorUtils::CanEditComponentInstance(Component, Private::Internal::GetParentComponent(Component), bAllowUserContructionScriptComps)
+			// Generally, snapshots should only restore snapshots that are editable. However, even if the root component is not editable, we still want to be able to restore the actor's transform.
+			// Ideally, we may not want to capture all component data in that case but for now that is exactly what we will do.
+			// For any future developer: we could add a IPropertyComparer implementation that marks all non-transform properties as unchanged.
+			&& (Component == Component->GetOwner()->GetRootComponent() || FComponentEditorUtils::CanEditComponentInstance(Component, Private::Internal::GetParentComponent(Component), bAllowUserContructionScriptComps))
 #endif
 	;
 	return bSomebodyAllowed || bIsAllowed;
 }
 
+bool UE::LevelSnapshots::Restorability::IsComponentRestorable(const ULevelSnapshot* Snapshot, const FSoftObjectPath& WorldActorPath, const UActorComponent* Component)
+{
+	return LevelSnapshots::Private::HasSavedComponentData(Snapshot->GetSerializedData(), WorldActorPath, Component)
+		&& IsComponentDesirableForCapture(Component);
+}
+
 bool UE::LevelSnapshots::Restorability::IsSubobjectClassDesirableForCapture(const UClass* SubobjectClass)
 {
 	SCOPED_SNAPSHOT_CORE_TRACE(IsSubobjectClassDesirableForCapture);
@@ -176,7 +187,7 @@ bool UE::LevelSnapshots::Restorability::IsSubobjectClassDesirableForCapture(cons
 
 bool UE::LevelSnapshots::Restorability::IsSubobjectDesirableForCapture(const UObject* Subobject)
 {
-	checkf(!Subobject->IsA<UClass>(), TEXT("Do you have a typo on your code?"));
+	checkf(!Subobject->IsA<UClass>(), TEXT("Do you have a typo in your code?"));
 	
 	if (const UActorComponent* Component = Cast<UActorComponent>(Subobject))
 	{
diff --git a/Engine/Plugins/VirtualProduction/LevelSnapshots/Source/LevelSnapshots/Private/Selection/ApplySnapshotFilter.cpp b/Engine/Plugins/VirtualProduction/LevelSnapshots/Source/LevelSnapshots/Private/Selection/ApplySnapshotFilter.cpp
--- a/Engine/Plugins/VirtualProduction/LevelSnapshots/Source/LevelSnapshots/Private/Selection/ApplySnapshotFilter.cpp
+++ b/Engine/Plugins/VirtualProduction/LevelSnapshots/Source/LevelSnapshots/Private/Selection/ApplySnapshotFilter.cpp
@@ -131,7 +131,7 @@ void UE::LevelSnapshots::Private::FApplySnapshotFilter::AnalyseComponentProperti
 	FAddedAndRemovedComponentInfo ComponentSelection;
 
 	const bool bShouldAnalyseUnmatchedComponents = Snapshot->GetSerializedData().SnapshotVersionInfo.GetSnapshotCustomVersion() >= FSnapshotCustomVersion::SubobjectSupport;
-	IterateComponents(DeserializedSnapshotActor, WorldActor,
+	IterateRestorableComponents(Snapshot, DeserializedSnapshotActor, WorldActor,
 		[this, &MapToAddTo](UActorComponent* SnapshotComp, UActorComponent* WorldComp)
 		{
 			FilterSubobjectPair(MapToAddTo, SnapshotComp, WorldComp);
diff --git a/Engine/Plugins/VirtualProduction/LevelSnapshots/Source/LevelSnapshots/Public/Data/SnapshotCustomVersion.h b/Engine/Plugins/VirtualProduction/LevelSnapshots/Source/LevelSnapshots/Public/Data/SnapshotCustomVersion.h
--- a/Engine/Plugins/VirtualProduction/LevelSnapshots/Source/LevelSnapshots/Public/Data/SnapshotCustomVersion.h
+++ b/Engine/Plugins/VirtualProduction/LevelSnapshots/Source/LevelSnapshots/Public/Data/SnapshotCustomVersion.h
@@ -35,6 +35,13 @@ namespace UE::LevelSnapshots
 			 */
 			SkipSubobjectReferencesWhenSerializingArchetypesRefactor = 6,
 
+			/**
+			 * If FComponentEditorUtils::CanEditComponentInstance returned false, the snapshot would not be captured.
+			 * Now, if FComponentEditorUtils::CanEditComponentInstance returns false the component will still be captured if it is a root component; otherwise the actor's transform cannot be restored.
+			 * No migration needs to take place. This entry exists only for future proofing (in case something breaks).
+			 */
+			NonEditableComponentsAreCaptured = 7,
+
 			// -----<new versions can be added above this line>-------------------------------------------------
 			VersionPlusOne,
 			LatestVersion = VersionPlusOne - 1
diff --git a/Engine/Plugins/VirtualProduction/LevelSnapshots/Source/LevelSnapshots/Public/Restorability/SnapshotRestorability.h b/Engine/Plugins/VirtualProduction/LevelSnapshots/Source/LevelSnapshots/Public/Restorability/SnapshotRestorability.h
--- a/Engine/Plugins/VirtualProduction/LevelSnapshots/Source/LevelSnapshots/Public/Restorability/SnapshotRestorability.h
+++ b/Engine/Plugins/VirtualProduction/LevelSnapshots/Source/LevelSnapshots/Public/Restorability/SnapshotRestorability.h
@@ -5,6 +5,7 @@
 #include "CoreMinimal.h"
 
 
+class ULevelSnapshot;
 class AActor;
 class FProperty;
 class UActorComponent;
@@ -26,7 +27,14 @@ namespace UE::LevelSnapshots::Restorability
 	
 	/* Is this component captured by the snapshot system? */
 	LEVELSNAPSHOTS_API bool IsComponentDesirableForCapture(const UActorComponent* Component);
-	
+	/**
+	 * Can this component be restored by the snapshot system?
+	 * @param Snapshot Needed to check versioning data, etc.
+	 * @param WorldActorPath Soft object path of the world actor
+	 * @param Component Component from either world or snapshot world
+	 */
+	LEVELSNAPSHOTS_API bool IsComponentRestorable(const ULevelSnapshot* Snapshot, const FSoftObjectPath& WorldActorPath, const UActorComponent* Component);
+
 	/* Is this subobject class captured by the snapshot system?*/
 	LEVELSNAPSHOTS_API bool IsSubobjectClassDesirableForCapture(const UClass* SubobjectClass);
 	/** Is this subobject captured by the snapshot system? */
diff --git a/Engine/Plugins/VirtualProduction/LevelSnapshots/Source/LevelSnapshotsEditor/Private/Views/Results/LevelSnapshotsEditorResultsHelpers.cpp b/Engine/Plugins/VirtualProduction/LevelSnapshots/Source/LevelSnapshotsEditor/Private/Views/Results/LevelSnapshotsEditorResultsHelpers.cpp
--- a/Engine/Plugins/VirtualProduction/LevelSnapshots/Source/LevelSnapshotsEditor/Private/Views/Results/LevelSnapshotsEditorResultsHelpers.cpp
+++ b/Engine/Plugins/VirtualProduction/LevelSnapshots/Source/LevelSnapshotsEditor/Private/Views/Results/LevelSnapshotsEditorResultsHelpers.cpp
@@ -556,7 +556,7 @@ TSharedRef<FComponentHierarchy> FLevelSnapshotsEditorResultsHelpers::BuildCompon
 			const bool ComponentContained = AllHierarchies.ContainsByPredicate(
 				[&SceneComponent](const TWeakPtr<FComponentHierarchy>& Hierarchy)
 				{
-					return Hierarchy.Pin()->Component == SceneComponent;
+					return Hierarchy.IsValid() && Hierarchy.Pin()->Component == SceneComponent;
 				});
 
 			if (!ComponentContained)
