Fix not checking TOptional value

diff --git a/Engine/Plugins/VirtualProduction/LevelSnapshots/Source/LevelSnapshots/Private/Data/LevelSnapshot.cpp b/Engine/Plugins/VirtualProduction/LevelSnapshots/Source/LevelSnapshots/Private/Data/LevelSnapshot.cpp
--- a/Engine/Plugins/VirtualProduction/LevelSnapshots/Source/LevelSnapshots/Private/Data/LevelSnapshot.cpp
+++ b/Engine/Plugins/VirtualProduction/LevelSnapshots/Source/LevelSnapshots/Private/Data/LevelSnapshot.cpp
@@ -133,8 +133,9 @@ bool ULevelSnapshot::HasChangedSinceSnapshotWasTaken(AActor* WorldActor)
 		
 		if (!bHasHashInfo || !bNeedsHash || !UE::LevelSnapshots::Private::HasMatchingHash(SavedActorData->Hash, WorldActor))
 		{
-			TOptional<TNonNullPtr<AActor>> DeserializedActor = GetDeserializedActor(WorldActor);
-			return HasOriginalChangedPropertiesSinceSnapshotWasTaken(DeserializedActor.GetValue(), WorldActor);
+			const TOptional<TNonNullPtr<AActor>> DeserializedActor = GetDeserializedActor(WorldActor);
+			UE_CLOG(!DeserializedActor, LogLevelSnapshots, Warning, TEXT("Failed to load snapshot actor %s. Actor will not be restored."), *WorldActor->GetPathName());
+			return DeserializedActor && HasOriginalChangedPropertiesSinceSnapshotWasTaken(DeserializedActor.GetValue(), WorldActor);
 		}
 
 		return false;
