Fix instanced components added to actor not getting detected by LevelSnapshots (detected by failing unit tests)

diff --git a/Engine/Plugins/VirtualProduction/LevelSnapshots/Source/LevelSnapshots/Private/Data/Util/EquivalenceUtil.cpp b/Engine/Plugins/VirtualProduction/LevelSnapshots/Source/LevelSnapshots/Private/Data/Util/EquivalenceUtil.cpp
--- a/Engine/Plugins/VirtualProduction/LevelSnapshots/Source/LevelSnapshots/Private/Data/Util/EquivalenceUtil.cpp
+++ b/Engine/Plugins/VirtualProduction/LevelSnapshots/Source/LevelSnapshots/Private/Data/Util/EquivalenceUtil.cpp
@@ -20,6 +20,7 @@
 #include "GameFramework/Actor.h"
 #include "UObject/TextProperty.h"
 #include "UObject/UnrealType.h"
+#include "WorldData/WorldDataUtil.h"
 
 namespace UE::LevelSnapshots::Private::Internal
 {
@@ -118,16 +119,17 @@ void UE::LevelSnapshots::Private::IterateRestorableComponents(ULevelSnapshot* Sn
 	const FSoftObjectPath WorldActorPath = WorldActor;
 	for (UActorComponent* WorldComp : WorldActor->GetComponents())
 	{
-		if (!Restorability::IsComponentRestorable(Snapshot, WorldActorPath, WorldComp))
+		if (!Restorability::IsComponentDesirableForCapture(WorldComp))
 		{
 			continue;
 		}
-		
-		if (UActorComponent* SnapshotMatchedComp = Internal::TryFindMatchingComponent(SnapshotActor, WorldComp))
+
+		UActorComponent* SnapshotMatchedComp = Internal::TryFindMatchingComponent(SnapshotActor, WorldComp);
+		if (SnapshotMatchedComp && HasSavedComponentData(Snapshot->GetSerializedData(), WorldActorPath, WorldComp))
 		{
 			OnComponentsMatched(SnapshotMatchedComp, WorldComp);
 		}
-		else
+		else if (!SnapshotMatchedComp)
 		{
 			OnWorldComponentUnmatched(WorldComp);
 		}
@@ -135,7 +137,8 @@ void UE::LevelSnapshots::Private::IterateRestorableComponents(ULevelSnapshot* Sn
 
 	for (UActorComponent* SnapshotComp : SnapshotActor->GetComponents())
 	{
-		if (Restorability::IsComponentRestorable(Snapshot, WorldActorPath, SnapshotComp)
+		if (Restorability::IsComponentDesirableForCapture(SnapshotComp)
+			&& HasSavedComponentData(Snapshot->GetSerializedData(), WorldActorPath, SnapshotComp)
 			&& Internal::TryFindMatchingComponent(WorldActor, SnapshotComp) == nullptr)
 		{
 			OnSnapshotComponentUnmatched(SnapshotComp);
diff --git a/Engine/Plugins/VirtualProduction/LevelSnapshots/Source/LevelSnapshots/Private/Data/Util/WorldData/WorldDataUtil.cpp b/Engine/Plugins/VirtualProduction/LevelSnapshots/Source/LevelSnapshots/Private/Data/Util/WorldData/WorldDataUtil.cpp
--- a/Engine/Plugins/VirtualProduction/LevelSnapshots/Source/LevelSnapshots/Private/Data/Util/WorldData/WorldDataUtil.cpp
+++ b/Engine/Plugins/VirtualProduction/LevelSnapshots/Source/LevelSnapshots/Private/Data/Util/WorldData/WorldDataUtil.cpp
@@ -267,7 +267,6 @@ bool UE::LevelSnapshots::Private::HasSavedComponentData(const FWorldSnapshotData
 		return false;
 	}
 
-	const FSoftObjectPath SearchedComponentPath = EditorOrSnapshotComponent;
 	const FString ComponentName = EditorOrSnapshotComponent->GetName();
 	const int32 ComponentDataIndex = [&WorldData, ActorData, ComponentName]() -> int32
 	{
diff --git a/Engine/Plugins/VirtualProduction/LevelSnapshots/Source/LevelSnapshots/Private/Restorability/SnapshotRestorability.cpp b/Engine/Plugins/VirtualProduction/LevelSnapshots/Source/LevelSnapshots/Private/Restorability/SnapshotRestorability.cpp
--- a/Engine/Plugins/VirtualProduction/LevelSnapshots/Source/LevelSnapshots/Private/Restorability/SnapshotRestorability.cpp
+++ b/Engine/Plugins/VirtualProduction/LevelSnapshots/Source/LevelSnapshots/Private/Restorability/SnapshotRestorability.cpp
@@ -173,12 +173,6 @@ bool UE::LevelSnapshots::Restorability::IsComponentDesirableForCapture(const UAc
 	return bSomebodyAllowed || bIsAllowed;
 }
 
-bool UE::LevelSnapshots::Restorability::IsComponentRestorable(const ULevelSnapshot* Snapshot, const FSoftObjectPath& WorldActorPath, const UActorComponent* Component)
-{
-	return LevelSnapshots::Private::HasSavedComponentData(Snapshot->GetSerializedData(), WorldActorPath, Component)
-		&& IsComponentDesirableForCapture(Component);
-}
-
 bool UE::LevelSnapshots::Restorability::IsSubobjectClassDesirableForCapture(const UClass* SubobjectClass)
 {
 	SCOPED_SNAPSHOT_CORE_TRACE(IsSubobjectClassDesirableForCapture);
diff --git a/Engine/Plugins/VirtualProduction/LevelSnapshots/Source/LevelSnapshots/Public/Restorability/SnapshotRestorability.h b/Engine/Plugins/VirtualProduction/LevelSnapshots/Source/LevelSnapshots/Public/Restorability/SnapshotRestorability.h
--- a/Engine/Plugins/VirtualProduction/LevelSnapshots/Source/LevelSnapshots/Public/Restorability/SnapshotRestorability.h
+++ b/Engine/Plugins/VirtualProduction/LevelSnapshots/Source/LevelSnapshots/Public/Restorability/SnapshotRestorability.h
@@ -27,13 +27,6 @@ namespace UE::LevelSnapshots::Restorability
 	
 	/* Is this component captured by the snapshot system? */
 	LEVELSNAPSHOTS_API bool IsComponentDesirableForCapture(const UActorComponent* Component);
-	/**
-	 * Can this component be restored by the snapshot system?
-	 * @param Snapshot Needed to check versioning data, etc.
-	 * @param WorldActorPath Soft object path of the world actor
-	 * @param Component Component from either world or snapshot world
-	 */
-	LEVELSNAPSHOTS_API bool IsComponentRestorable(const ULevelSnapshot* Snapshot, const FSoftObjectPath& WorldActorPath, const UActorComponent* Component);
 
 	/* Is this subobject class captured by the snapshot system?*/
 	LEVELSNAPSHOTS_API bool IsSubobjectClassDesirableForCapture(const UClass* SubobjectClass);
