Fix material instances showing up as changed if though they're not

diff --git a/Engine/Plugins/VirtualProduction/LevelSnapshots/Source/LevelSnapshots/Private/Data/Util/EquivalenceUtil.cpp b/Engine/Plugins/VirtualProduction/LevelSnapshots/Source/LevelSnapshots/Private/Data/Util/EquivalenceUtil.cpp
--- a/Engine/Plugins/VirtualProduction/LevelSnapshots/Source/LevelSnapshots/Private/Data/Util/EquivalenceUtil.cpp
+++ b/Engine/Plugins/VirtualProduction/LevelSnapshots/Source/LevelSnapshots/Private/Data/Util/EquivalenceUtil.cpp
@@ -445,7 +445,8 @@ bool UE::LevelSnapshots::Private::AreReferencesEquivalent(ULevelSnapshot* Snapsh
 	{
 		return AreActorsEquivalent(SnapshotPropertyValue, OriginalActorReference, Snapshot->GetSerializedData(), Snapshot->GetCache());
 	}
-	
+
+	// Fast path for components and other subobjects of actors.
 	const bool bIsWorldObject = SnapshotPropertyValue->IsInA(UWorld::StaticClass()) && OriginalPropertyValue->IsInA(UWorld::StaticClass());
 	if (bIsWorldObject)
 	{
@@ -453,9 +454,22 @@ bool UE::LevelSnapshots::Private::AreReferencesEquivalent(ULevelSnapshot* Snapsh
 		const bool bAreComponents = SnapshotPropertyValue->IsA<UActorComponent>() && OriginalPropertyValue->IsA<UActorComponent>();
 		return (bAreComponents && Internal::HaveSameNames(SnapshotPropertyValue, OriginalPropertyValue, Snapshot->GetCache()))
 			|| !Restorability::IsSubobjectDesirableForCapture(OriginalPropertyValue)
+			// Iterates through all properties -> possibly slow
 			|| (!bAreComponents && Internal::AreEquivalentSubobjects(Snapshot, SnapshotPropertyValue, OriginalPropertyValue, SnapshotActor, OriginalActor));
 	}
-	
+
+	// Example: two materials in the GetTransientPackage()
+	const bool bIsSnapshotObjectTransient = (SnapshotPropertyValue->GetPackage()->GetFlags() & RF_Transient) != 0;
+	const bool bIsWorldObjectTransient = (OriginalPropertyValue->GetPackage()->GetFlags() & RF_Transient) != 0;
+	const bool bAreTransient = bIsSnapshotObjectTransient || bIsWorldObjectTransient;
+	if (bAreTransient)
+	{
+		return !Restorability::IsSubobjectDesirableForCapture(OriginalPropertyValue)
+			// Iterates through all properties -> possibly slow
+			|| !Internal::HaveDifferentPropertyValues(Snapshot, SnapshotPropertyValue, OriginalPropertyValue, SnapshotActor, OriginalActor);
+	}
+
+	// Pointing to assets in asset registry, i.e. no subobjects?
 	return SnapshotPropertyValue == OriginalPropertyValue;
 }
 
diff --git a/Engine/Plugins/VirtualProduction/LevelSnapshots/Source/LevelSnapshots/Private/Restorability/EngineTypesRestorationFence.cpp b/Engine/Plugins/VirtualProduction/LevelSnapshots/Source/LevelSnapshots/Private/Restorability/EngineTypesRestorationFence.cpp
--- a/Engine/Plugins/VirtualProduction/LevelSnapshots/Source/LevelSnapshots/Private/Restorability/EngineTypesRestorationFence.cpp
+++ b/Engine/Plugins/VirtualProduction/LevelSnapshots/Source/LevelSnapshots/Private/Restorability/EngineTypesRestorationFence.cpp
@@ -12,6 +12,7 @@
 #include "Algo/Transform.h"
 #include "Components/SkeletalMeshComponent.h"
 #include "EngineUtils.h"
+#include "MaterialInstanceSupport.h"
 #include "Engine/Brush.h"
 #include "GameFramework/Actor.h"
 #include "Materials/MaterialInstance.h"
@@ -126,5 +127,6 @@ namespace UE::LevelSnapshots::Private::EngineTypesRestorationFence
 		SpecialActorPropertySupport::Register(Module);
 		AttachParentShowsTransformPropertiesFix::Register(Module);
 		LevelInstanceRestoration::Register(Module);
+		MaterialInstanceSupport::Register(Module);
 	}
 }
diff --git a/Engine/Plugins/VirtualProduction/LevelSnapshots/Source/LevelSnapshots/Private/Restorability/MaterialInstanceSupport.cpp b/Engine/Plugins/VirtualProduction/LevelSnapshots/Source/LevelSnapshots/Private/Restorability/MaterialInstanceSupport.cpp
new file mode 100644
--- /dev/null
+++ b/Engine/Plugins/VirtualProduction/LevelSnapshots/Source/LevelSnapshots/Private/Restorability/MaterialInstanceSupport.cpp
@@ -0,0 +1,97 @@
+// Copyright Epic Games, Inc. All Rights Reserved.
+
+#include "MaterialInstanceSupport.h"
+
+#include "Params/PropertyComparisonParams.h"
+#include "Util/EquivalenceUtil.h"
+
+#include "Materials/MaterialInstance.h"
+#include "UObject/UnrealType.h"
+
+namespace UE::LevelSnapshots::Private::MaterialInstanceSupport
+{
+	/** Sometimes the parameters (e.g. ScalarParameterValues) compared are out of order causing the material to show up as changed. */
+	class FMaterialInstancePropertyComparer : public IPropertyComparer
+	{
+		const FProperty* Property_ScalarParameterValues = UMaterialInstance::StaticClass()->FindPropertyByName(GET_MEMBER_NAME_CHECKED(UMaterialInstance, ScalarParameterValues));
+		const FProperty* Property_VectorParameterValues = UMaterialInstance::StaticClass()->FindPropertyByName(GET_MEMBER_NAME_CHECKED(UMaterialInstance, VectorParameterValues));
+		const FProperty* Property_DoubleVectorParameterValues = UMaterialInstance::StaticClass()->FindPropertyByName(GET_MEMBER_NAME_CHECKED(UMaterialInstance, DoubleVectorParameterValues));
+		const FProperty* Property_TextureParameterValues = UMaterialInstance::StaticClass()->FindPropertyByName(GET_MEMBER_NAME_CHECKED(UMaterialInstance, TextureParameterValues));
+		const FProperty* Property_RuntimeVirtualTextureParameterValues = UMaterialInstance::StaticClass()->FindPropertyByName(GET_MEMBER_NAME_CHECKED(UMaterialInstance, RuntimeVirtualTextureParameterValues));
+		const FProperty* Property_FontParameterValues = UMaterialInstance::StaticClass()->FindPropertyByName(GET_MEMBER_NAME_CHECKED(UMaterialInstance, FontParameterValues));
+	public:
+
+		FMaterialInstancePropertyComparer()
+		{
+			check(Property_ScalarParameterValues
+				&& Property_VectorParameterValues
+				&& Property_DoubleVectorParameterValues
+				&& Property_TextureParameterValues
+				&& Property_RuntimeVirtualTextureParameterValues
+				&& Property_FontParameterValues);
+		}
+		
+		virtual EPropertyComparison ShouldConsiderPropertyEqual(const FPropertyComparisonParams& Params) const override
+		{
+			// Could be written "nicer" but this code will be called often so make it fast to compute the if
+			const bool bIsCustomProperty = Params.LeafProperty == Property_ScalarParameterValues
+				|| Params.LeafProperty == Property_VectorParameterValues
+				|| Params.LeafProperty == Property_DoubleVectorParameterValues
+				|| Params.LeafProperty == Property_TextureParameterValues
+				|| Params.LeafProperty == Property_RuntimeVirtualTextureParameterValues
+				|| Params.LeafProperty == Property_FontParameterValues;
+			
+			if (bIsCustomProperty)
+			{
+				return AreArraysTransmutation(Params, CastFieldChecked<FArrayProperty>(Params.LeafProperty), Params.SnapshotContainer, Params.WorldContainer)
+					? EPropertyComparison::TreatEqual
+					: EPropertyComparison::TreatUnequal;
+			}
+			
+			return EPropertyComparison::CheckNormally;
+		}
+
+		bool AreArraysTransmutation(const FPropertyComparisonParams& Params, const FArrayProperty* Property, void* SnapshotContainerPtr, void* WorldContainerPtr) const
+		{
+			void* SnapshotValuePtr = Property->ContainerPtrToValuePtr<void>(SnapshotContainerPtr);
+			void* WorldValuePtr = Property->ContainerPtrToValuePtr<void>(WorldContainerPtr);
+
+			FScriptArrayHelper SnapshotArray(Property, SnapshotValuePtr);
+			FScriptArrayHelper WorldArray(Property, WorldValuePtr);
+			if (SnapshotArray.Num() != WorldArray.Num())
+			{
+				return false;
+			}
+
+			for (int32 i = 0; i < SnapshotArray.Num(); ++i)
+			{
+				void* SnapshotElementValuePtr = SnapshotArray.GetRawPtr(i);
+				if (!ContainsEquivalentElement(Params, Property, SnapshotElementValuePtr, WorldArray))
+				{
+					return false;
+				}
+			}
+			
+			return true;
+		}
+
+		bool ContainsEquivalentElement(const FPropertyComparisonParams& Params, const FArrayProperty* Property, void* SnapshotElementValuePtr, FScriptArrayHelper& WorldArray) const
+		{
+			for (int32 j = 0; j < WorldArray.Num(); ++j)
+			{
+				void* WorldElementValuePtr = WorldArray.GetRawPtr(j);
+				if (AreSnapshotAndOriginalPropertiesEquivalent(Params.Snapshot, Property->Inner, SnapshotElementValuePtr, WorldElementValuePtr, Params.SnapshotActor, Params.WorldActor))
+				{
+					return true;
+				}
+			}
+			return false;
+		}
+	};
+	
+	void Register(FLevelSnapshotsModule& Module)
+	{
+		const TSharedRef<FMaterialInstancePropertyComparer> MaterialInstancePropertyComparer = MakeShared<FMaterialInstancePropertyComparer>();
+		Module.RegisterPropertyComparer(UMaterialInstance::StaticClass(), MaterialInstancePropertyComparer);
+	}
+}
diff --git a/Engine/Plugins/VirtualProduction/LevelSnapshots/Source/LevelSnapshots/Private/Restorability/MaterialInstanceSupport.h b/Engine/Plugins/VirtualProduction/LevelSnapshots/Source/LevelSnapshots/Private/Restorability/MaterialInstanceSupport.h
new file mode 100644
--- /dev/null
+++ b/Engine/Plugins/VirtualProduction/LevelSnapshots/Source/LevelSnapshots/Private/Restorability/MaterialInstanceSupport.h
@@ -0,0 +1,11 @@
+// Copyright Epic Games, Inc. All Rights Reserved.
+
+#pragma once
+
+#include "CoreMinimal.h"
+#include "LevelSnapshotsModule.h"
+
+namespace UE::LevelSnapshots::Private::MaterialInstanceSupport
+{
+	void Register(FLevelSnapshotsModule& Module);
+}
