Fix further encountered instances of TObjectPtr<T> used with T that is not derived from UObject.  A subsequent change will be disallowing this at compile time.

diff --git a/Engine/Plugins/Editor/ConsoleVariablesEditor/Source/ConsoleVariablesEditor/Private/Views/List/SConsoleVariablesEditorListValueInput.cpp b/Engine/Plugins/Editor/ConsoleVariablesEditor/Source/ConsoleVariablesEditor/Private/Views/List/SConsoleVariablesEditorListValueInput.cpp
--- a/Engine/Plugins/Editor/ConsoleVariablesEditor/Source/ConsoleVariablesEditor/Private/Views/List/SConsoleVariablesEditorListValueInput.cpp
+++ b/Engine/Plugins/Editor/ConsoleVariablesEditor/Source/ConsoleVariablesEditor/Private/Views/List/SConsoleVariablesEditorListValueInput.cpp
@@ -26,7 +26,7 @@ TSharedRef<SConsoleVariablesEditorListValueInput> SConsoleVariablesEditorListVal
 	{
 		const TSharedPtr<FConsoleVariablesEditorCommandInfo> PinnedInfo = PinnedItem->GetCommandInfo().Pin();
 
-		if (const TObjectPtr<IConsoleVariable> Variable = PinnedInfo->GetConsoleVariablePtr())
+		if (const IConsoleVariable* Variable = PinnedInfo->GetConsoleVariablePtr())
 		{
 			if (Variable->IsVariableFloat())
 			{
diff --git a/Engine/Plugins/VirtualProduction/LevelSnapshots/Source/LevelSnapshotsEditor/Private/Misc/LevelSnapshotsEditorCustomWidgetGenerator.cpp b/Engine/Plugins/VirtualProduction/LevelSnapshots/Source/LevelSnapshotsEditor/Private/Misc/LevelSnapshotsEditorCustomWidgetGenerator.cpp
--- a/Engine/Plugins/VirtualProduction/LevelSnapshots/Source/LevelSnapshotsEditor/Private/Misc/LevelSnapshotsEditorCustomWidgetGenerator.cpp
+++ b/Engine/Plugins/VirtualProduction/LevelSnapshots/Source/LevelSnapshotsEditor/Private/Misc/LevelSnapshotsEditorCustomWidgetGenerator.cpp
@@ -26,7 +26,7 @@ void LevelSnapshotsEditorCustomWidgetGenerator::CreateRowsForPropertiesNotHandle
 		return;
 	}
 	
-	const TObjectPtr<FProperty> Property = InFieldPath.Get();
+	const FProperty* Property = InFieldPath.Get();
 
 	if (!ensureAlwaysMsgf(Property,
 		TEXT("%hs: Property was not able to be retrieved from InFieldPath. Please ensure InFieldPath is valid before calling this function. InFieldPath: %s"),
@@ -62,7 +62,7 @@ void LevelSnapshotsEditorCustomWidgetGenerator::CreateRowsForPropertiesNotHandle
 			return LOCTEXT("LevelSnapshotsEditorResults_AttachParentInvalidText","No Attach Parent found.").ToString();
 		};
 
-		if (const TObjectPtr<FObjectPropertyBase> ObjectProperty = CastField<FObjectPropertyBase>(Property))
+		if (const FObjectPropertyBase* ObjectProperty = CastField<FObjectPropertyBase>(Property))
 		{
 			CustomSnapshotWidget = GenerateObjectPropertyWidget(ObjectProperty->GetObjectPropertyValue(SnapshotPropertyValue), WidgetTextEditLambda);
 			CustomWorldWidget = GenerateObjectPropertyWidget(ObjectProperty->GetObjectPropertyValue(WorldPropertyValue), WidgetTextEditLambda);
@@ -92,7 +92,7 @@ void LevelSnapshotsEditorCustomWidgetGenerator::CreateRowsForPropertiesNotHandle
 	InDirectParentRow.Pin()->AddToChildRows(NewProperty);
 }
 
-void* LevelSnapshotsEditorCustomWidgetGenerator::GetPropertyValueFromProperty(const TObjectPtr<FProperty> InProperty, const TObjectPtr<UObject> InObject)
+void* LevelSnapshotsEditorCustomWidgetGenerator::GetPropertyValueFromProperty(const FProperty* InProperty, const TObjectPtr<UObject> InObject)
 {
 	void* PropertyValue = nullptr;
 	
@@ -139,7 +139,7 @@ void* LevelSnapshotsEditorCustomWidgetGenerator::GetPropertyValueFromProperty(co
 }
 
 TSharedPtr<SWidget> LevelSnapshotsEditorCustomWidgetGenerator::GenerateGenericPropertyWidget(
-	const TObjectPtr<FProperty> InProperty, const void* InPropertyValue, const TObjectPtr<UObject> InObject,
+	const FProperty* InProperty, const void* InPropertyValue, const TObjectPtr<UObject> InObject,
 	const TFunction<FString(const FString&, const UObject*)> InWidgetTextEditLambda)
 {
 	// We should have found a value ptr in the methods above, but we ensure in case of a missed scenario
@@ -181,7 +181,7 @@ TSharedPtr<SWidget> LevelSnapshotsEditorCustomWidgetGenerator::GenerateObjectPro
 }
 
 TSharedPtr<SWidget> LevelSnapshotsEditorCustomWidgetGenerator::DeterminePropertyTypeAndReturnWidget(
-	const TObjectPtr<FProperty> InProperty, const void* InPropertyValue, const TObjectPtr<UObject> InObject)
+	const FProperty* InProperty, const void* InPropertyValue, const TObjectPtr<UObject> InObject)
 {
 	if (!ensure(InPropertyValue))
 	{
@@ -192,9 +192,9 @@ TSharedPtr<SWidget> LevelSnapshotsEditorCustomWidgetGenerator::DetermineProperty
 	const FText ToolTipText = FText::FromString(InProperty->GetClass()->GetName());
 
 	// Then we'll iterate over relevant property types to see what kind of widget we should create
-	if (const TObjectPtr<FNumericProperty> NumericProperty = CastField<FNumericProperty>(InProperty))
+	if (const FNumericProperty* NumericProperty = CastField<FNumericProperty>(InProperty))
 	{
-		if (const TObjectPtr<FFloatProperty> FloatProperty = CastField<FFloatProperty>(InProperty))
+		if (const FFloatProperty* FloatProperty = CastField<FFloatProperty>(InProperty))
 		{
 			const float Value = FloatProperty->GetFloatingPointPropertyValue(InPropertyValue);
 			return SNew(SNumericEntryBox<float>)
@@ -204,7 +204,7 @@ TSharedPtr<SWidget> LevelSnapshotsEditorCustomWidgetGenerator::DetermineProperty
 				.MinSliderValue(Value)
 				.MaxSliderValue(Value);
 		}
-		else if (const TObjectPtr<FDoubleProperty> DoubleProperty = CastField<FDoubleProperty>(InProperty))
+		else if (const FDoubleProperty* DoubleProperty = CastField<FDoubleProperty>(InProperty))
 		{
 			const double Value = DoubleProperty->GetFloatingPointPropertyValue(InPropertyValue);
 			return SNew(SNumericEntryBox<double>)
@@ -225,13 +225,13 @@ TSharedPtr<SWidget> LevelSnapshotsEditorCustomWidgetGenerator::DetermineProperty
 				.MaxSliderValue(Value);
 		}
 	}
-	else if (const TObjectPtr<FBoolProperty> BoolProperty = CastField<FBoolProperty>(InProperty))
+	else if (const FBoolProperty* BoolProperty = CastField<FBoolProperty>(InProperty))
 	{
 		return SNew(SCheckBox)
 				.IsChecked(BoolProperty->GetPropertyValue(InPropertyValue))
 				.ToolTipText(ToolTipText);
 	}
-	else if (const TObjectPtr<FObjectPropertyBase> ObjectProperty = CastField<FObjectPropertyBase>(InProperty))
+	else if (const FObjectPropertyBase* ObjectProperty = CastField<FObjectPropertyBase>(InProperty))
 	{
 		return GenerateObjectPropertyWidget(ObjectProperty->GetObjectPropertyValue(InPropertyValue));
 	}
diff --git a/Engine/Plugins/VirtualProduction/LevelSnapshots/Source/LevelSnapshotsEditor/Private/Misc/LevelSnapshotsEditorCustomWidgetGenerator.h b/Engine/Plugins/VirtualProduction/LevelSnapshots/Source/LevelSnapshotsEditor/Private/Misc/LevelSnapshotsEditorCustomWidgetGenerator.h
--- a/Engine/Plugins/VirtualProduction/LevelSnapshots/Source/LevelSnapshotsEditor/Private/Misc/LevelSnapshotsEditorCustomWidgetGenerator.h
+++ b/Engine/Plugins/VirtualProduction/LevelSnapshots/Source/LevelSnapshotsEditor/Private/Misc/LevelSnapshotsEditorCustomWidgetGenerator.h
@@ -25,17 +25,17 @@ namespace LevelSnapshotsEditorCustomWidgetGenerator
 		const TWeakPtr<SLevelSnapshotsEditorResults>& InResultsView,
 		const TWeakPtr<FLevelSnapshotsEditorResultsRow>& InDirectParentRow);
 
-	void* GetPropertyValueFromProperty(const TObjectPtr<FProperty> InProperty, const TObjectPtr<UObject> InObject);
+	void* GetPropertyValueFromProperty(const FProperty* InProperty, const TObjectPtr<UObject> InObject);
 
 	/** Used for properties which are not specifically handled in another method. Returns an STextBlock widget with the value exported to text. */
 	TSharedPtr<SWidget> GenerateGenericPropertyWidget(
-		const TObjectPtr<FProperty> InProperty, const void* InPropertyValue, const TObjectPtr<UObject> InObject,
+		const FProperty* InProperty, const void* InPropertyValue, const TObjectPtr<UObject> InObject,
 		const TFunction<FString(const FString&, const UObject*)> InWidgetTextEditLambda = nullptr);
 
 	TSharedPtr<SWidget> GenerateObjectPropertyWidget(
 		const TObjectPtr<UObject> InObject, const TFunction<FString(const FString&, const UObject*)> InWidgetTextEditLambda = nullptr);
 
-	TSharedPtr<SWidget> DeterminePropertyTypeAndReturnWidget(const TObjectPtr<FProperty> InProperty, const void* InPropertyValue, const TObjectPtr<UObject> InObject);
+	TSharedPtr<SWidget> DeterminePropertyTypeAndReturnWidget(const FProperty* InProperty, const void* InPropertyValue, const TObjectPtr<UObject> InObject);
 
 	TSharedRef<SWidget> MakeComboBoxWithSelection(const FString& InString, const FText& InToolTipText);
 };
