Remove unnecessary Public and Private entries for the current module being added to PublicIncludePaths or PrivateIncludePaths

#rnx

diff --git a/Engine/Plugins/Animation/DeformerGraph/Source/OptimusCore/OptimusCore.Build.cs b/Engine/Plugins/Animation/DeformerGraph/Source/OptimusCore/OptimusCore.Build.cs
--- a/Engine/Plugins/Animation/DeformerGraph/Source/OptimusCore/OptimusCore.Build.cs
+++ b/Engine/Plugins/Animation/DeformerGraph/Source/OptimusCore/OptimusCore.Build.cs
@@ -8,12 +8,6 @@ namespace UnrealBuildTool.Rules
 		{
 			NumIncludedBytesPerUnityCPPOverride = 688128; // best unity size found from using UBT ProfileUnitySizes mode
 
-			PrivateIncludePaths.AddRange(
-                new string[] {
-                    "OptimusCore/Private",
-				}
-            );
-
             PublicDependencyModuleNames.AddRange(
 				new string[]
 				{
diff --git a/Engine/Plugins/Animation/DeformerGraph/Source/OptimusDeveloper/OptimusDeveloper.Build.cs b/Engine/Plugins/Animation/DeformerGraph/Source/OptimusDeveloper/OptimusDeveloper.Build.cs
--- a/Engine/Plugins/Animation/DeformerGraph/Source/OptimusDeveloper/OptimusDeveloper.Build.cs
+++ b/Engine/Plugins/Animation/DeformerGraph/Source/OptimusDeveloper/OptimusDeveloper.Build.cs
@@ -6,12 +6,6 @@ namespace UnrealBuildTool.Rules
     {
         public OptimusDeveloper(ReadOnlyTargetRules Target) : base(Target)
         {
-	        PrivateIncludePaths.AddRange(
-                new string[] {
-					"OptimusDeveloper/Private",
-				}
-            );
-
             PublicDependencyModuleNames.AddRange(
 				new string[]
 				{
diff --git a/Engine/Plugins/Animation/DeformerGraph/Source/OptimusEditor/OptimusEditor.Build.cs b/Engine/Plugins/Animation/DeformerGraph/Source/OptimusEditor/OptimusEditor.Build.cs
--- a/Engine/Plugins/Animation/DeformerGraph/Source/OptimusEditor/OptimusEditor.Build.cs
+++ b/Engine/Plugins/Animation/DeformerGraph/Source/OptimusEditor/OptimusEditor.Build.cs
@@ -8,7 +8,6 @@ namespace UnrealBuildTool.Rules
         {
 			PrivateIncludePaths.AddRange(new string[] {
 				System.IO.Path.Combine(GetModuleDirectory("OptimusCore"), "Private"),
-				System.IO.Path.Combine(GetModuleDirectory("OptimusEditor"), "Private"),
 				System.IO.Path.Combine(GetModuleDirectory("PropertyEditor"), "Private"),
 			});
 
diff --git a/Engine/Plugins/Animation/DeformerGraph/Source/OptimusSettings/OptimusSettings.Build.cs b/Engine/Plugins/Animation/DeformerGraph/Source/OptimusSettings/OptimusSettings.Build.cs
--- a/Engine/Plugins/Animation/DeformerGraph/Source/OptimusSettings/OptimusSettings.Build.cs
+++ b/Engine/Plugins/Animation/DeformerGraph/Source/OptimusSettings/OptimusSettings.Build.cs
@@ -6,12 +6,6 @@ namespace UnrealBuildTool.Rules
     {
         public OptimusSettings(ReadOnlyTargetRules Target) : base(Target)
         {
-            PrivateIncludePaths.AddRange(
-                new string[] {
-					"OptimusSettings/Private",
-				}
-            );
-
             PublicDependencyModuleNames.AddRange(
 				new string[]
 				{
diff --git a/Engine/Plugins/Bridge/Source/Bridge/Bridge.Build.cs b/Engine/Plugins/Bridge/Source/Bridge/Bridge.Build.cs
--- a/Engine/Plugins/Bridge/Source/Bridge/Bridge.Build.cs
+++ b/Engine/Plugins/Bridge/Source/Bridge/Bridge.Build.cs
@@ -46,11 +46,6 @@ namespace UnrealBuildTool.Rules
 					"MegascansPlugin",
 					"ApplicationCore",
 					});
-
-			PrivateIncludePaths.AddRange(
-				new string[] {
-					Path.Combine(ModuleDirectory, "Private"),
-					});
 			}
 		}
 }
diff --git a/Engine/Plugins/Editor/CryptoKeys/Source/CryptoKeys/CryptoKeys.Build.cs b/Engine/Plugins/Editor/CryptoKeys/Source/CryptoKeys/CryptoKeys.Build.cs
--- a/Engine/Plugins/Editor/CryptoKeys/Source/CryptoKeys/CryptoKeys.Build.cs
+++ b/Engine/Plugins/Editor/CryptoKeys/Source/CryptoKeys/CryptoKeys.Build.cs
@@ -6,8 +6,6 @@ public class CryptoKeys : ModuleRules
 {
 	public CryptoKeys(ReadOnlyTargetRules Target) : base(Target)
 	{
-		PrivateIncludePaths.Add("CryptoKeys/Classes");
-
 		PublicDependencyModuleNames.AddRange(
 			new string[]
 			{
diff --git a/Engine/Plugins/Editor/ObjectMixer/LightMixer/Source/LightMixer/LightMixer.build.cs b/Engine/Plugins/Editor/ObjectMixer/LightMixer/Source/LightMixer/LightMixer.build.cs
--- a/Engine/Plugins/Editor/ObjectMixer/LightMixer/Source/LightMixer/LightMixer.build.cs
+++ b/Engine/Plugins/Editor/ObjectMixer/LightMixer/Source/LightMixer/LightMixer.build.cs
@@ -9,8 +9,6 @@ public class LightMixer : ModuleRules
 	{
 		PCHUsage = ModuleRules.PCHUsageMode.UseExplicitOrSharedPCHs;
 		
-		PublicIncludePaths.Add(Path.Combine(ModuleDirectory, "Public"));
-
 		PublicDependencyModuleNames.AddRange(
 			new string[]
 			{
diff --git a/Engine/Plugins/Experimental/Animation/SkeletalMeshModelingTools/Source/SkeletalMeshModelingTools/SkeletalMeshModelingTools.Build.cs b/Engine/Plugins/Experimental/Animation/SkeletalMeshModelingTools/Source/SkeletalMeshModelingTools/SkeletalMeshModelingTools.Build.cs
--- a/Engine/Plugins/Experimental/Animation/SkeletalMeshModelingTools/Source/SkeletalMeshModelingTools/SkeletalMeshModelingTools.Build.cs
+++ b/Engine/Plugins/Experimental/Animation/SkeletalMeshModelingTools/Source/SkeletalMeshModelingTools/SkeletalMeshModelingTools.Build.cs
@@ -6,12 +6,6 @@ public class SkeletalMeshModelingTools : ModuleRules
 {
 	public SkeletalMeshModelingTools(ReadOnlyTargetRules Target) : base(Target)
 	{
-		PrivateIncludePaths.AddRange(
-			new string[] {
-				"SkeletalMeshModelingTools/Private",
-			}
-		);
-
 		PrivateDependencyModuleNames.AddRange(
 			new string[]
 			{
diff --git a/Engine/Plugins/Experimental/AppleVision/Source/AppleVisionBlueprintSupport/AppleVisionBlueprintSupport.Build.cs b/Engine/Plugins/Experimental/AppleVision/Source/AppleVisionBlueprintSupport/AppleVisionBlueprintSupport.Build.cs
--- a/Engine/Plugins/Experimental/AppleVision/Source/AppleVisionBlueprintSupport/AppleVisionBlueprintSupport.Build.cs
+++ b/Engine/Plugins/Experimental/AppleVision/Source/AppleVisionBlueprintSupport/AppleVisionBlueprintSupport.Build.cs
@@ -10,7 +10,6 @@ namespace UnrealBuildTool.Rules
             PrivateIncludePaths.AddRange(
                 new string[] {
                     "AppleVision/Public",
-                    "AppleVisionBlueprintSupport/Classes"
                     // ... add other private include paths required here ...
                 }
                 );
diff --git a/Engine/Plugins/Experimental/ChaosCloth/Source/Chaos/ChaosCloth.Build.cs b/Engine/Plugins/Experimental/ChaosCloth/Source/Chaos/ChaosCloth.Build.cs
--- a/Engine/Plugins/Experimental/ChaosCloth/Source/Chaos/ChaosCloth.Build.cs
+++ b/Engine/Plugins/Experimental/ChaosCloth/Source/Chaos/ChaosCloth.Build.cs
@@ -6,8 +6,6 @@ public class ChaosCloth : ModuleRules
 {
     public ChaosCloth(ReadOnlyTargetRules Target) : base(Target)
 	{
-        PrivateIncludePaths.Add("Chaos/Private");
-
         PublicDependencyModuleNames.AddRange(
             new string[] {
                 "Core",
diff --git a/Engine/Plugins/Experimental/ChaosClothEditor/Source/Chaos/ChaosClothEditor.Build.cs b/Engine/Plugins/Experimental/ChaosClothEditor/Source/Chaos/ChaosClothEditor.Build.cs
--- a/Engine/Plugins/Experimental/ChaosClothEditor/Source/Chaos/ChaosClothEditor.Build.cs
+++ b/Engine/Plugins/Experimental/ChaosClothEditor/Source/Chaos/ChaosClothEditor.Build.cs
@@ -6,8 +6,6 @@ public class ChaosClothEditor : ModuleRules
 {
     public ChaosClothEditor(ReadOnlyTargetRules Target) : base(Target)
     {
-        PrivateIncludePaths.Add("Chaos/Private");
-
         PrivateDependencyModuleNames.AddRange(
             new string[] {
                 "Core",
diff --git a/Engine/Plugins/Experimental/CommonUI/Source/CommonInput/CommonInput.Build.cs b/Engine/Plugins/Experimental/CommonUI/Source/CommonInput/CommonInput.Build.cs
--- a/Engine/Plugins/Experimental/CommonUI/Source/CommonInput/CommonInput.Build.cs
+++ b/Engine/Plugins/Experimental/CommonUI/Source/CommonInput/CommonInput.Build.cs
@@ -28,13 +28,6 @@ public class CommonInput : ModuleRules
             }
 		);
 
-        PrivateIncludePaths.AddRange(
-			new string[]
-			{
-                "CommonInput/Private",
-			}
-		);
-
         if (Target.Type == TargetType.Editor)
         {
             PublicDependencyModuleNames.AddRange(
diff --git a/Engine/Plugins/Experimental/CommonUI/Source/CommonUIEditor/CommonUIEditor.Build.cs b/Engine/Plugins/Experimental/CommonUI/Source/CommonUIEditor/CommonUIEditor.Build.cs
--- a/Engine/Plugins/Experimental/CommonUI/Source/CommonUIEditor/CommonUIEditor.Build.cs
+++ b/Engine/Plugins/Experimental/CommonUI/Source/CommonUIEditor/CommonUIEditor.Build.cs
@@ -35,7 +35,6 @@ public class CommonUIEditor : ModuleRules
 			new string[]
 			{
 				System.IO.Path.Combine(GetModuleDirectory("CommonUI"), "Private"),
-				System.IO.Path.Combine(GetModuleDirectory("CommonUIEditor"), "Private"),
 			}
 		);
 
diff --git a/Engine/Plugins/Experimental/Enterprise/DatasmithRuntime/Source/DatasmithRuntime.Build.cs b/Engine/Plugins/Experimental/Enterprise/DatasmithRuntime/Source/DatasmithRuntime.Build.cs
--- a/Engine/Plugins/Experimental/Enterprise/DatasmithRuntime/Source/DatasmithRuntime.Build.cs
+++ b/Engine/Plugins/Experimental/Enterprise/DatasmithRuntime/Source/DatasmithRuntime.Build.cs
@@ -11,8 +11,6 @@ public class DatasmithRuntime : ModuleRules
 		OptimizeCode = CodeOptimization.InShippingBuildsOnly;
 		PCHUsage = ModuleRules.PCHUsageMode.UseExplicitOrSharedPCHs;
 
-		PrivateIncludePaths.Add("Private");
-
 		PublicDependencyModuleNames.AddRange(
 			new string[] {
 				"DatasmithContent",
diff --git a/Engine/Plugins/Experimental/Landmass/Source/Editor/LandmassEditor.Build.cs b/Engine/Plugins/Experimental/Landmass/Source/Editor/LandmassEditor.Build.cs
--- a/Engine/Plugins/Experimental/Landmass/Source/Editor/LandmassEditor.Build.cs
+++ b/Engine/Plugins/Experimental/Landmass/Source/Editor/LandmassEditor.Build.cs
@@ -7,8 +7,6 @@ public class LandmassEditor : ModuleRules
 	public LandmassEditor(ReadOnlyTargetRules Target) : base(Target)
 	{
 		PCHUsage = ModuleRules.PCHUsageMode.UseExplicitOrSharedPCHs;
-		
-		PrivateIncludePaths.Add("Editor/Private");
 
 		PublicDependencyModuleNames.AddRange(
 			new string[] {
diff --git a/Engine/Plugins/Experimental/Landmass/Source/Runtime/Landmass.Build.cs b/Engine/Plugins/Experimental/Landmass/Source/Runtime/Landmass.Build.cs
--- a/Engine/Plugins/Experimental/Landmass/Source/Runtime/Landmass.Build.cs
+++ b/Engine/Plugins/Experimental/Landmass/Source/Runtime/Landmass.Build.cs
@@ -8,8 +8,6 @@ public class Landmass : ModuleRules
 	{
 		PCHUsage = ModuleRules.PCHUsageMode.UseExplicitOrSharedPCHs;
 		
-		PrivateIncludePaths.Add("Runtime/Private");
-
 		PublicDependencyModuleNames.AddRange(
 			new string[] {
 				"Core",
diff --git a/Engine/Plugins/Experimental/ModularGameplay/Source/ModularGameplay/ModularGameplay.Build.cs b/Engine/Plugins/Experimental/ModularGameplay/Source/ModularGameplay/ModularGameplay.Build.cs
--- a/Engine/Plugins/Experimental/ModularGameplay/Source/ModularGameplay/ModularGameplay.Build.cs
+++ b/Engine/Plugins/Experimental/ModularGameplay/Source/ModularGameplay/ModularGameplay.Build.cs
@@ -6,13 +6,7 @@ public class ModularGameplay : ModuleRules
 {
 	public ModularGameplay(ReadOnlyTargetRules Target) : base(Target)
 	{
-		PCHUsage = ModuleRules.PCHUsageMode.UseExplicitOrSharedPCHs;				
-		
-		PrivateIncludePaths.AddRange(
-			new string[] {
-                "ModularGameplay/Private"
-            }
-			);
+		PCHUsage = ModuleRules.PCHUsageMode.UseExplicitOrSharedPCHs;
 			
 		
 		PublicDependencyModuleNames.AddRange(
diff --git a/Engine/Plugins/Experimental/Water/Source/Editor/WaterEditor.Build.cs b/Engine/Plugins/Experimental/Water/Source/Editor/WaterEditor.Build.cs
--- a/Engine/Plugins/Experimental/Water/Source/Editor/WaterEditor.Build.cs
+++ b/Engine/Plugins/Experimental/Water/Source/Editor/WaterEditor.Build.cs
@@ -7,8 +7,6 @@ public class WaterEditor : ModuleRules
 	public WaterEditor(ReadOnlyTargetRules Target) : base(Target)
 	{
 		PCHUsage = ModuleRules.PCHUsageMode.UseExplicitOrSharedPCHs;
-		
-        PrivateIncludePaths.Add("Editor/Private");
 
 		PrivateDependencyModuleNames.AddRange(
 			new string[] {
diff --git a/Engine/Plugins/FX/CascadeToNiagaraConverter/Source/CascadeToNiagaraConverter/CascadeToNiagaraConverter.Build.cs b/Engine/Plugins/FX/CascadeToNiagaraConverter/Source/CascadeToNiagaraConverter/CascadeToNiagaraConverter.Build.cs
--- a/Engine/Plugins/FX/CascadeToNiagaraConverter/Source/CascadeToNiagaraConverter/CascadeToNiagaraConverter.Build.cs
+++ b/Engine/Plugins/FX/CascadeToNiagaraConverter/Source/CascadeToNiagaraConverter/CascadeToNiagaraConverter.Build.cs
@@ -7,10 +7,6 @@ public class CascadeToNiagaraConverter : ModuleRules
 	public CascadeToNiagaraConverter(ReadOnlyTargetRules Target) : base(Target)
 	{
         PCHUsage = ModuleRules.PCHUsageMode.UseExplicitOrSharedPCHs;
-
-        PrivateIncludePaths.AddRange(new string[] {
-				"CascadeToNiagaraConverter/Private",
-		});
 		
 		PrivateDependencyModuleNames.AddRange(
 			new string[]
diff --git a/Engine/Plugins/Media/AjaMedia/Source/AjaMediaFactory/AjaMediaFactory.Build.cs b/Engine/Plugins/Media/AjaMedia/Source/AjaMediaFactory/AjaMediaFactory.Build.cs
--- a/Engine/Plugins/Media/AjaMedia/Source/AjaMediaFactory/AjaMediaFactory.Build.cs
+++ b/Engine/Plugins/Media/AjaMedia/Source/AjaMediaFactory/AjaMediaFactory.Build.cs
@@ -23,11 +23,6 @@ namespace UnrealBuildTool.Rules
 					"AjaMedia",
 				});
 
-			PrivateIncludePaths.AddRange(
-				new string[] {
-					"AjaMediaFactory/Private",
-				});
-
 			PublicDependencyModuleNames.AddRange(
 				new string[] {
 					"Core",
diff --git a/Engine/Plugins/Media/AjaMedia/Source/AjaMediaOutput/AjaMediaOutput.Build.cs b/Engine/Plugins/Media/AjaMedia/Source/AjaMediaOutput/AjaMediaOutput.Build.cs
--- a/Engine/Plugins/Media/AjaMedia/Source/AjaMediaOutput/AjaMediaOutput.Build.cs
+++ b/Engine/Plugins/Media/AjaMedia/Source/AjaMediaOutput/AjaMediaOutput.Build.cs
@@ -14,13 +14,6 @@ namespace UnrealBuildTool.Rules
 					"MediaIOCore",
 				});
 
-			PrivateIncludePaths.AddRange(
-				new string[]
-				{
-					"AjaMediaOutput/Private"
-				}
-			);
-
 			PrivateDependencyModuleNames.AddRange(
 				new string[]
 				{
diff --git a/Engine/Plugins/Media/AndroidCamera/Source/AndroidCameraEditor/AndroidCameraEditor.Build.cs b/Engine/Plugins/Media/AndroidCamera/Source/AndroidCameraEditor/AndroidCameraEditor.Build.cs
--- a/Engine/Plugins/Media/AndroidCamera/Source/AndroidCameraEditor/AndroidCameraEditor.Build.cs
+++ b/Engine/Plugins/Media/AndroidCamera/Source/AndroidCameraEditor/AndroidCameraEditor.Build.cs
@@ -14,11 +14,6 @@ namespace UnrealBuildTool.Rules
 					"MediaAssets",
 					"UnrealEd",
 				});
-
-			PrivateIncludePaths.AddRange(
-				new string[] {
-					"AndroidCameraEditor/Private",
-				});
 		}
 	}
 }
diff --git a/Engine/Plugins/Media/AndroidCamera/Source/AndroidCameraFactory/AndroidCameraFactory.Build.cs b/Engine/Plugins/Media/AndroidCamera/Source/AndroidCameraFactory/AndroidCameraFactory.Build.cs
--- a/Engine/Plugins/Media/AndroidCamera/Source/AndroidCameraFactory/AndroidCameraFactory.Build.cs
+++ b/Engine/Plugins/Media/AndroidCamera/Source/AndroidCameraFactory/AndroidCameraFactory.Build.cs
@@ -22,11 +22,6 @@ namespace UnrealBuildTool.Rules
 					"Media",
 				});
 
-			PrivateIncludePaths.AddRange(
-				new string[] {
-					"AndroidCameraFactory/Private",
-				});
-
 			PublicDependencyModuleNames.AddRange(
 				new string[] {
 					"Core",
diff --git a/Engine/Plugins/Media/AndroidMedia/Source/AndroidMediaEditor/AndroidMediaEditor.Build.cs b/Engine/Plugins/Media/AndroidMedia/Source/AndroidMediaEditor/AndroidMediaEditor.Build.cs
--- a/Engine/Plugins/Media/AndroidMedia/Source/AndroidMediaEditor/AndroidMediaEditor.Build.cs
+++ b/Engine/Plugins/Media/AndroidMedia/Source/AndroidMediaEditor/AndroidMediaEditor.Build.cs
@@ -14,11 +14,6 @@ namespace UnrealBuildTool.Rules
 					"MediaAssets",
 					"UnrealEd",
 				});
-
-			PrivateIncludePaths.AddRange(
-				new string[] {
-					"AndroidMediaEditor/Private",
-				});
 		}
 	}
 }
diff --git a/Engine/Plugins/Media/AndroidMedia/Source/AndroidMediaFactory/AndroidMediaFactory.Build.cs b/Engine/Plugins/Media/AndroidMedia/Source/AndroidMediaFactory/AndroidMediaFactory.Build.cs
--- a/Engine/Plugins/Media/AndroidMedia/Source/AndroidMediaFactory/AndroidMediaFactory.Build.cs
+++ b/Engine/Plugins/Media/AndroidMedia/Source/AndroidMediaFactory/AndroidMediaFactory.Build.cs
@@ -22,11 +22,6 @@ namespace UnrealBuildTool.Rules
 					"Media",
 				});
 
-			PrivateIncludePaths.AddRange(
-				new string[] {
-					"AndroidMediaFactory/Private",
-				});
-
 			PublicDependencyModuleNames.AddRange(
 				new string[] {
 					"Core",
diff --git a/Engine/Plugins/Media/AvfMedia/Source/AvfMediaEditor/AvfMediaEditor.Build.cs b/Engine/Plugins/Media/AvfMedia/Source/AvfMediaEditor/AvfMediaEditor.Build.cs
--- a/Engine/Plugins/Media/AvfMedia/Source/AvfMediaEditor/AvfMediaEditor.Build.cs
+++ b/Engine/Plugins/Media/AvfMedia/Source/AvfMediaEditor/AvfMediaEditor.Build.cs
@@ -14,11 +14,6 @@ namespace UnrealBuildTool.Rules
 					"MediaAssets",
 					"UnrealEd",
 				});
-
-			PrivateIncludePaths.AddRange(
-				new string[] {
-					"AvfMediaEditor/Private",
-				});
 		}
 	}
 }
diff --git a/Engine/Plugins/Media/AvfMedia/Source/AvfMediaFactory/AvfMediaFactory.Build.cs b/Engine/Plugins/Media/AvfMedia/Source/AvfMediaFactory/AvfMediaFactory.Build.cs
--- a/Engine/Plugins/Media/AvfMedia/Source/AvfMediaFactory/AvfMediaFactory.Build.cs
+++ b/Engine/Plugins/Media/AvfMedia/Source/AvfMediaFactory/AvfMediaFactory.Build.cs
@@ -22,11 +22,6 @@ namespace UnrealBuildTool.Rules
 					"Media",
 				});
 
-			PrivateIncludePaths.AddRange(
-				new string[] {
-					"AvfMediaFactory/Private",
-				});
-
 			PublicDependencyModuleNames.AddRange(
 				new string[] {
 					"Core",
diff --git a/Engine/Plugins/Media/BlackmagicMedia/Source/BlackmagicMediaFactory/BlackmagicMediaFactory.Build.cs b/Engine/Plugins/Media/BlackmagicMedia/Source/BlackmagicMediaFactory/BlackmagicMediaFactory.Build.cs
--- a/Engine/Plugins/Media/BlackmagicMedia/Source/BlackmagicMediaFactory/BlackmagicMediaFactory.Build.cs
+++ b/Engine/Plugins/Media/BlackmagicMedia/Source/BlackmagicMediaFactory/BlackmagicMediaFactory.Build.cs
@@ -23,11 +23,6 @@ namespace UnrealBuildTool.Rules
 					"BlackmagicMedia",
 				});
 
-			PrivateIncludePaths.AddRange(
-				new string[] {
-					"BlackmagicMediaFactory/Private",
-				});
-
 			PublicDependencyModuleNames.AddRange(
 				new string[] {
 					"Core",
diff --git a/Engine/Plugins/Media/BlackmagicMedia/Source/BlackmagicMediaOutput/BlackmagicMediaOutput.Build.cs b/Engine/Plugins/Media/BlackmagicMedia/Source/BlackmagicMediaOutput/BlackmagicMediaOutput.Build.cs
--- a/Engine/Plugins/Media/BlackmagicMedia/Source/BlackmagicMediaOutput/BlackmagicMediaOutput.Build.cs
+++ b/Engine/Plugins/Media/BlackmagicMedia/Source/BlackmagicMediaOutput/BlackmagicMediaOutput.Build.cs
@@ -14,13 +14,6 @@ namespace UnrealBuildTool.Rules
 					"MediaIOCore",
 				});
 
-			PrivateIncludePaths.AddRange(
-				new string[]
-				{
-					"BlackmagicMediaOutput/Private"
-				}
-			);
-
 			PrivateDependencyModuleNames.AddRange(
 				new string[]
 				{
diff --git a/Engine/Plugins/Media/ElectraPlayer/Source/ElectraPlayerFactory/ElectraPlayerFactory.Build.cs b/Engine/Plugins/Media/ElectraPlayer/Source/ElectraPlayerFactory/ElectraPlayerFactory.Build.cs
--- a/Engine/Plugins/Media/ElectraPlayer/Source/ElectraPlayerFactory/ElectraPlayerFactory.Build.cs
+++ b/Engine/Plugins/Media/ElectraPlayer/Source/ElectraPlayerFactory/ElectraPlayerFactory.Build.cs
@@ -30,11 +30,6 @@ namespace UnrealBuildTool.Rules
                 PublicDefinitions.Add("UE_PLATFORM_ELECTRAPLAYER=1");
             }
 
-            PrivateIncludePaths.AddRange(
-                new string[] {
-                    "ElectraPlayerFactory/Private",
-                });
-
             PublicDependencyModuleNames.AddRange(
                 new string[] {
                     "Core",
diff --git a/Engine/Plugins/Media/ElectraUtil/Source/ElectraHTTPStream/ElectraHTTPStream.Build.cs b/Engine/Plugins/Media/ElectraUtil/Source/ElectraHTTPStream/ElectraHTTPStream.Build.cs
--- a/Engine/Plugins/Media/ElectraUtil/Source/ElectraHTTPStream/ElectraHTTPStream.Build.cs
+++ b/Engine/Plugins/Media/ElectraUtil/Source/ElectraHTTPStream/ElectraHTTPStream.Build.cs
@@ -74,11 +74,6 @@ public class ElectraHTTPStream : ModuleRules
 	{
 		PublicDefinitions.Add("ELECTRA_HTTPSTREAM_PACKAGE=1");
 
-		PrivateIncludePaths.AddRange(
-			new string[] {
-				"ElectraHTTPStream/Private",
-			});
-
 		PublicDependencyModuleNames.AddRange(
 			new string[] {
 				"Core"
diff --git a/Engine/Plugins/Media/HLMedia/Source/HLMedia/HLMedia.Build.cs b/Engine/Plugins/Media/HLMedia/Source/HLMedia/HLMedia.Build.cs
--- a/Engine/Plugins/Media/HLMedia/Source/HLMedia/HLMedia.Build.cs
+++ b/Engine/Plugins/Media/HLMedia/Source/HLMedia/HLMedia.Build.cs
@@ -37,11 +37,6 @@ public class HLMedia : ModuleRules
                 "Media",
             });
 
-        PrivateIncludePaths.AddRange(
-            new string[] {
-                "HLMedia/Private",
-            });
-
         PrivateDependencyModuleNames.AddRange(new string[] 
         {
             "HLMediaLibrary",
diff --git a/Engine/Plugins/Media/HLMedia/Source/HLMediaFactory/HLMediaFactory.Build.cs b/Engine/Plugins/Media/HLMedia/Source/HLMediaFactory/HLMediaFactory.Build.cs
--- a/Engine/Plugins/Media/HLMedia/Source/HLMediaFactory/HLMediaFactory.Build.cs
+++ b/Engine/Plugins/Media/HLMedia/Source/HLMediaFactory/HLMediaFactory.Build.cs
@@ -24,11 +24,6 @@ namespace UnrealBuildTool.Rules
                     "HLMedia",
                 });
 
-            PrivateIncludePaths.AddRange(
-                new string[] {
-                    "HLMediaFactory/Private",
-                });
-
             DynamicallyLoadedModuleNames.Add("HLMedia");
         }
     }
diff --git a/Engine/Plugins/Media/ImgMedia/Source/ImgMediaFactory/ImgMediaFactory.Build.cs b/Engine/Plugins/Media/ImgMedia/Source/ImgMediaFactory/ImgMediaFactory.Build.cs
--- a/Engine/Plugins/Media/ImgMedia/Source/ImgMediaFactory/ImgMediaFactory.Build.cs
+++ b/Engine/Plugins/Media/ImgMedia/Source/ImgMediaFactory/ImgMediaFactory.Build.cs
@@ -22,11 +22,6 @@ namespace UnrealBuildTool.Rules
 					"ImgMedia",
 				});
 
-			PrivateIncludePaths.AddRange(
-				new string[] {
-					"ImgMediaFactory/Private",
-				});
-
 			PublicDependencyModuleNames.AddRange(
 				new string[] {
 					"Core",
diff --git a/Engine/Plugins/Media/MfMedia/Source/MfMediaEditor/MfMediaEditor.Build.cs b/Engine/Plugins/Media/MfMedia/Source/MfMediaEditor/MfMediaEditor.Build.cs
--- a/Engine/Plugins/Media/MfMedia/Source/MfMediaEditor/MfMediaEditor.Build.cs
+++ b/Engine/Plugins/Media/MfMedia/Source/MfMediaEditor/MfMediaEditor.Build.cs
@@ -14,11 +14,6 @@ namespace UnrealBuildTool.Rules
 					"MediaAssets",
 					"UnrealEd",
 				});
-
-			PrivateIncludePaths.AddRange(
-				new string[] {
-					"MfMediaEditor/Private",
-				});
 		}
 	}
 }
diff --git a/Engine/Plugins/Media/MfMedia/Source/MfMediaFactory/MfMediaFactory.Build.cs b/Engine/Plugins/Media/MfMedia/Source/MfMediaFactory/MfMediaFactory.Build.cs
--- a/Engine/Plugins/Media/MfMedia/Source/MfMediaFactory/MfMediaFactory.Build.cs
+++ b/Engine/Plugins/Media/MfMedia/Source/MfMediaFactory/MfMediaFactory.Build.cs
@@ -24,11 +24,6 @@ namespace UnrealBuildTool.Rules
 					"MfMedia",
 				});
 
-			PrivateIncludePaths.AddRange(
-				new string[] {
-					"MfMediaFactory/Private",
-				});
-
 			if (DoAllowHTTPSPlayback())
 			{
 				PrivateDefinitions.Add("MFMEDIAFACTORY_ALLOW_HTTPS=1");
diff --git a/Engine/Plugins/Media/WebMMedia/Source/WebMMediaFactory/WebMMediaFactory.Build.cs b/Engine/Plugins/Media/WebMMedia/Source/WebMMediaFactory/WebMMediaFactory.Build.cs
--- a/Engine/Plugins/Media/WebMMedia/Source/WebMMediaFactory/WebMMediaFactory.Build.cs
+++ b/Engine/Plugins/Media/WebMMedia/Source/WebMMediaFactory/WebMMediaFactory.Build.cs
@@ -21,11 +21,6 @@ namespace UnrealBuildTool.Rules
 					"Media",
 				});
 
-			PrivateIncludePaths.AddRange(
-				new string[] {
-					"WebMMediaFactory/Private",
-				});
-
 			PublicDependencyModuleNames.AddRange(
 				new string[] {
 					"Core",
diff --git a/Engine/Plugins/Media/WmfMedia/Source/WmfMediaEditor/WmfMediaEditor.Build.cs b/Engine/Plugins/Media/WmfMedia/Source/WmfMediaEditor/WmfMediaEditor.Build.cs
--- a/Engine/Plugins/Media/WmfMedia/Source/WmfMediaEditor/WmfMediaEditor.Build.cs
+++ b/Engine/Plugins/Media/WmfMedia/Source/WmfMediaEditor/WmfMediaEditor.Build.cs
@@ -14,11 +14,6 @@ namespace UnrealBuildTool.Rules
 					"MediaAssets",
 					"UnrealEd",
 				});
-
-			PrivateIncludePaths.AddRange(
-				new string[] {
-					"WmfMediaEditor/Private",
-				});
 		}
 	}
 }
diff --git a/Engine/Plugins/Media/WmfMedia/Source/WmfMediaFactory/WmfMediaFactory.Build.cs b/Engine/Plugins/Media/WmfMedia/Source/WmfMediaFactory/WmfMediaFactory.Build.cs
--- a/Engine/Plugins/Media/WmfMedia/Source/WmfMediaFactory/WmfMediaFactory.Build.cs
+++ b/Engine/Plugins/Media/WmfMedia/Source/WmfMediaFactory/WmfMediaFactory.Build.cs
@@ -24,11 +24,6 @@ namespace UnrealBuildTool.Rules
 					"WmfMedia",
 				});
 
-			PrivateIncludePaths.AddRange(
-				new string[] {
-					"WmfMediaFactory/Private",
-				});
-
 			PublicDependencyModuleNames.AddRange(
 				new string[] {
 					"Core",
diff --git a/Engine/Plugins/Messaging/MessagingDebugger/Source/MessagingDebugger/MessagingDebugger.Build.cs b/Engine/Plugins/Messaging/MessagingDebugger/Source/MessagingDebugger/MessagingDebugger.Build.cs
--- a/Engine/Plugins/Messaging/MessagingDebugger/Source/MessagingDebugger/MessagingDebugger.Build.cs
+++ b/Engine/Plugins/Messaging/MessagingDebugger/Source/MessagingDebugger/MessagingDebugger.Build.cs
@@ -24,7 +24,6 @@ public class MessagingDebugger : ModuleRules
 
 		PrivateIncludePaths.AddRange(
 			new string[] {
-				"MessagingDebugger/Private",
 				"MessagingDebugger/Private/Models",
 				"MessagingDebugger/Private/Styles",
 				"MessagingDebugger/Private/Widgets",
diff --git a/Engine/Plugins/Messaging/TcpMessaging/Source/TcpMessaging/TcpMessaging.Build.cs b/Engine/Plugins/Messaging/TcpMessaging/Source/TcpMessaging/TcpMessaging.Build.cs
--- a/Engine/Plugins/Messaging/TcpMessaging/Source/TcpMessaging/TcpMessaging.Build.cs
+++ b/Engine/Plugins/Messaging/TcpMessaging/Source/TcpMessaging/TcpMessaging.Build.cs
@@ -34,7 +34,6 @@ namespace UnrealBuildTool.Rules
 
 			PrivateIncludePaths.AddRange(
 				new string[] {
-					"TcpMessaging/Private",
 					"TcpMessaging/Private/Settings",
 					"TcpMessaging/Private/Transport",
 				});
diff --git a/Engine/Plugins/MovieScene/CustomizableSequencerTracks/Source/CustomizableSequencerTracksEditor/CustomizableSequencerTracksEditor.Build.cs b/Engine/Plugins/MovieScene/CustomizableSequencerTracks/Source/CustomizableSequencerTracksEditor/CustomizableSequencerTracksEditor.Build.cs
--- a/Engine/Plugins/MovieScene/CustomizableSequencerTracks/Source/CustomizableSequencerTracksEditor/CustomizableSequencerTracksEditor.Build.cs
+++ b/Engine/Plugins/MovieScene/CustomizableSequencerTracks/Source/CustomizableSequencerTracksEditor/CustomizableSequencerTracksEditor.Build.cs
@@ -8,12 +8,6 @@ public class CustomizableSequencerTracksEditor : ModuleRules
 	{
 		PCHUsage = ModuleRules.PCHUsageMode.UseExplicitOrSharedPCHs;
 
-		PrivateIncludePaths.AddRange(
-			new string[] {
-				"CustomizableSequencerTracksEditor/Private",
-			}
-		);
-
 		PrivateDependencyModuleNames.AddRange(
 			new string[]
 			{
diff --git a/Engine/Plugins/MovieScene/LevelSequenceEditor/Source/LevelSequenceEditor/LevelSequenceEditor.Build.cs b/Engine/Plugins/MovieScene/LevelSequenceEditor/Source/LevelSequenceEditor/LevelSequenceEditor.Build.cs
--- a/Engine/Plugins/MovieScene/LevelSequenceEditor/Source/LevelSequenceEditor/LevelSequenceEditor.Build.cs
+++ b/Engine/Plugins/MovieScene/LevelSequenceEditor/Source/LevelSequenceEditor/LevelSequenceEditor.Build.cs
@@ -67,8 +67,6 @@ public class LevelSequenceEditor : ModuleRules
 
         PrivateIncludePaths.AddRange(
             new string[] {
-            	"LevelSequenceEditor/Public",
-				"LevelSequenceEditor/Private",
 				"LevelSequenceEditor/Private/AssetTools",
 				"LevelSequenceEditor/Private/Factories",
                 "LevelSequenceEditor/Private/Misc",
diff --git a/Engine/Plugins/MovieScene/ReplayTracks/Source/ReplayTracksEditor/ReplayTracksEditor.Build.cs b/Engine/Plugins/MovieScene/ReplayTracks/Source/ReplayTracksEditor/ReplayTracksEditor.Build.cs
--- a/Engine/Plugins/MovieScene/ReplayTracks/Source/ReplayTracksEditor/ReplayTracksEditor.Build.cs
+++ b/Engine/Plugins/MovieScene/ReplayTracks/Source/ReplayTracksEditor/ReplayTracksEditor.Build.cs
@@ -23,11 +23,5 @@ public class ReplayTracksEditor : ModuleRules
 				"TimeManagement",
 			}
 		);
-
-		PrivateIncludePaths.AddRange(
-			new string[] {
-				"ReplayTracksEditor/Private",
-			}
-		);
 	}
 }
diff --git a/Engine/Plugins/MovieScene/SequencerScripting/Source/SequencerScripting/SequencerScripting.Build.cs b/Engine/Plugins/MovieScene/SequencerScripting/Source/SequencerScripting/SequencerScripting.Build.cs
--- a/Engine/Plugins/MovieScene/SequencerScripting/Source/SequencerScripting/SequencerScripting.Build.cs
+++ b/Engine/Plugins/MovieScene/SequencerScripting/Source/SequencerScripting/SequencerScripting.Build.cs
@@ -15,7 +15,6 @@ public class SequencerScripting : ModuleRules
 
 		PrivateIncludePaths.AddRange(
 			new string[] {
-				"SequencerScripting/Private",
 				"SequencerScripting/Private/ExtensionLibraries",
                 "SequencerScripting/Private/KeysAndChannels",
             }
diff --git a/Engine/Plugins/MovieScene/SequencerScripting/Source/SequencerScriptingEditor/SequencerScriptingEditor.Build.cs b/Engine/Plugins/MovieScene/SequencerScripting/Source/SequencerScriptingEditor/SequencerScriptingEditor.Build.cs
--- a/Engine/Plugins/MovieScene/SequencerScripting/Source/SequencerScriptingEditor/SequencerScriptingEditor.Build.cs
+++ b/Engine/Plugins/MovieScene/SequencerScripting/Source/SequencerScriptingEditor/SequencerScriptingEditor.Build.cs
@@ -15,7 +15,6 @@ public class SequencerScriptingEditor : ModuleRules
 
 		PrivateIncludePaths.AddRange(
 			new string[] {
-				"SequencerScriptingEditor/Private",
                 "SequencerScripting/Private",
                 "../../../../Source/Editor/UnrealEd/Private", // TODO: Fix this, for now it's needed for the fbx exporter
             }
diff --git a/Engine/Plugins/MovieScene/TemplateSequence/Source/TemplateSequenceEditor/TemplateSequenceEditor.Build.cs b/Engine/Plugins/MovieScene/TemplateSequence/Source/TemplateSequenceEditor/TemplateSequenceEditor.Build.cs
--- a/Engine/Plugins/MovieScene/TemplateSequence/Source/TemplateSequenceEditor/TemplateSequenceEditor.Build.cs
+++ b/Engine/Plugins/MovieScene/TemplateSequence/Source/TemplateSequenceEditor/TemplateSequenceEditor.Build.cs
@@ -31,12 +31,6 @@ public class TemplateSequenceEditor : ModuleRules
 			}
 		);
 
-		PrivateIncludePaths.AddRange(
-			new string[] {
-				"TemplateSequenceEditor/Private",
-			}
-		);
-
 		var DynamicModuleNames = new string[] {
 			"LevelEditor",
 			"PropertyEditor",
diff --git a/Engine/Plugins/Online/VoiceChat/EOSVoiceChat/Source/EOSVoiceChat/EOSVoiceChat.Build.cs b/Engine/Plugins/Online/VoiceChat/EOSVoiceChat/Source/EOSVoiceChat/EOSVoiceChat.Build.cs
--- a/Engine/Plugins/Online/VoiceChat/EOSVoiceChat/Source/EOSVoiceChat/EOSVoiceChat.Build.cs
+++ b/Engine/Plugins/Online/VoiceChat/EOSVoiceChat/Source/EOSVoiceChat/EOSVoiceChat.Build.cs
@@ -11,8 +11,6 @@ namespace UnrealBuildTool.Rules
 		{
 			PCHUsage = ModuleRules.PCHUsageMode.UseExplicitOrSharedPCHs;
 
-			PublicIncludePaths.Add(Path.Combine(ModuleDirectory, "Public"));
-
 			PublicDependencyModuleNames.AddRange(
 				new string[]
 				{
diff --git a/Engine/Plugins/Runtime/AR/Google/GoogleARCore/Source/GoogleARCoreRendering/GoogleARCoreRendering.Build.cs b/Engine/Plugins/Runtime/AR/Google/GoogleARCore/Source/GoogleARCoreRendering/GoogleARCoreRendering.Build.cs
--- a/Engine/Plugins/Runtime/AR/Google/GoogleARCore/Source/GoogleARCoreRendering/GoogleARCoreRendering.Build.cs
+++ b/Engine/Plugins/Runtime/AR/Google/GoogleARCore/Source/GoogleARCoreRendering/GoogleARCoreRendering.Build.cs
@@ -8,7 +8,6 @@ public class GoogleARCoreRendering : ModuleRules
 	{
 		PrivateIncludePaths.AddRange(new string[]
 		{
-			"GoogleARCoreRendering/Private",
 			"GoogleARCoreBase/Private",
 		});
 			
diff --git a/Engine/Plugins/Runtime/AppleImageUtils/Source/AppleImageUtilsBlueprintSupport/AppleImageUtilsBlueprintSupport.Build.cs b/Engine/Plugins/Runtime/AppleImageUtils/Source/AppleImageUtilsBlueprintSupport/AppleImageUtilsBlueprintSupport.Build.cs
--- a/Engine/Plugins/Runtime/AppleImageUtils/Source/AppleImageUtilsBlueprintSupport/AppleImageUtilsBlueprintSupport.Build.cs
+++ b/Engine/Plugins/Runtime/AppleImageUtils/Source/AppleImageUtilsBlueprintSupport/AppleImageUtilsBlueprintSupport.Build.cs
@@ -16,7 +16,6 @@ namespace UnrealBuildTool.Rules
             PrivateIncludePaths.AddRange(
                 new string[] {
                     "AppleImageUtils/Public",
-                    "AppleImageUtilsBlueprintSupport/Classes"
                     // ... add other private include paths required here ...
                 }
                 );
diff --git a/Engine/Plugins/Runtime/AudioModulation/Source/AudioModulationEditor/AudioModulationEditor.build.cs b/Engine/Plugins/Runtime/AudioModulation/Source/AudioModulationEditor/AudioModulationEditor.build.cs
--- a/Engine/Plugins/Runtime/AudioModulation/Source/AudioModulationEditor/AudioModulationEditor.build.cs
+++ b/Engine/Plugins/Runtime/AudioModulation/Source/AudioModulationEditor/AudioModulationEditor.build.cs
@@ -40,7 +40,6 @@ public class AudioModulationEditor : ModuleRules
 		PrivateIncludePaths.AddRange(
 			new string[]
 			{
-				"AudioModulationEditor/Private",
 				"AudioModulation/Private"
 			}
 		);
diff --git a/Engine/Plugins/Runtime/MicrosoftSpatialAudio/Source/MicrosoftSpatialSound/MicrosoftSpatialSound.Build.cs b/Engine/Plugins/Runtime/MicrosoftSpatialAudio/Source/MicrosoftSpatialSound/MicrosoftSpatialSound.Build.cs
--- a/Engine/Plugins/Runtime/MicrosoftSpatialAudio/Source/MicrosoftSpatialSound/MicrosoftSpatialSound.Build.cs
+++ b/Engine/Plugins/Runtime/MicrosoftSpatialAudio/Source/MicrosoftSpatialSound/MicrosoftSpatialSound.Build.cs
@@ -44,7 +44,6 @@ public class MicrosoftSpatialSound : ModuleRules
 		
 		PrivateIncludePaths.AddRange(
 			new string[] {
-				"MicrosoftSpatialSound/Private",
 				System.IO.Path.Combine(EngineDir, "Source/ThirdParty/SpatialAudioClientInterop"),
 			}
 			);
diff --git a/Engine/Plugins/Runtime/Oculus/OculusAudio/Source/OculusAudio/OculusAudio.Build.cs b/Engine/Plugins/Runtime/Oculus/OculusAudio/Source/OculusAudio/OculusAudio.Build.cs
--- a/Engine/Plugins/Runtime/Oculus/OculusAudio/Source/OculusAudio/OculusAudio.Build.cs
+++ b/Engine/Plugins/Runtime/Oculus/OculusAudio/Source/OculusAudio/OculusAudio.Build.cs
@@ -17,8 +17,6 @@ namespace UnrealBuildTool.Rules
 
 			PrivateIncludePaths.AddRange(
 				new string[] {
-					// Relative to Engine\Plugins\Runtime\Oculus\OculusAudio\Source
-					"OculusAudio/Private",
 					System.IO.Path.Combine(GetModuleDirectory("AudioMixer"), "Private"),
 
 					// ... add other private include paths required here ...
diff --git a/Engine/Plugins/Runtime/Oculus/OculusAudio/Source/OculusAudioEditor/OculusAudioEditor.Build.cs b/Engine/Plugins/Runtime/Oculus/OculusAudio/Source/OculusAudioEditor/OculusAudioEditor.Build.cs
--- a/Engine/Plugins/Runtime/Oculus/OculusAudio/Source/OculusAudioEditor/OculusAudioEditor.Build.cs
+++ b/Engine/Plugins/Runtime/Oculus/OculusAudio/Source/OculusAudioEditor/OculusAudioEditor.Build.cs
@@ -16,7 +16,6 @@ namespace UnrealBuildTool.Rules
 
             PrivateIncludePaths.AddRange(
                 new string[] {
-                    "OculusAudioEditor/Public",
                     "OculusAudio/Private",
                     "OculusAudio/Public",
  					// ... add other private include paths required here ...
diff --git a/Engine/Plugins/Runtime/Oculus/OculusVR/Source/OculusEditor/OculusEditor.Build.cs b/Engine/Plugins/Runtime/Oculus/OculusVR/Source/OculusEditor/OculusEditor.Build.cs
--- a/Engine/Plugins/Runtime/Oculus/OculusVR/Source/OculusEditor/OculusEditor.Build.cs
+++ b/Engine/Plugins/Runtime/Oculus/OculusVR/Source/OculusEditor/OculusEditor.Build.cs
@@ -34,7 +34,6 @@ public class OculusEditor : ModuleRules
 		PrivateIncludePaths.AddRange(
 				new string[] {
 					// Relative to Engine\Plugins\Runtime\Oculus\OculusVR\Source
-					"OculusEditor/Private",
 					"OculusHMD/Private",
 				});
 
diff --git a/Engine/Plugins/Runtime/SignificanceManager/Source/SignificanceManager/SignificanceManager.Build.cs b/Engine/Plugins/Runtime/SignificanceManager/Source/SignificanceManager/SignificanceManager.Build.cs
--- a/Engine/Plugins/Runtime/SignificanceManager/Source/SignificanceManager/SignificanceManager.Build.cs
+++ b/Engine/Plugins/Runtime/SignificanceManager/Source/SignificanceManager/SignificanceManager.Build.cs
@@ -14,11 +14,5 @@ public class SignificanceManager : ModuleRules
 				"EngineSettings"
 			}
 		);
-
-        PrivateIncludePaths.AddRange(
-            new string[] {
-                "SignificanceManager/Private"
-            })
-		;
 	}
 }
diff --git a/Engine/Plugins/Runtime/Steam/SteamAudio/Source/SteamAudioEditor/SteamAudioEditor.Build.cs b/Engine/Plugins/Runtime/Steam/SteamAudio/Source/SteamAudioEditor/SteamAudioEditor.Build.cs
--- a/Engine/Plugins/Runtime/Steam/SteamAudio/Source/SteamAudioEditor/SteamAudioEditor.Build.cs
+++ b/Engine/Plugins/Runtime/Steam/SteamAudio/Source/SteamAudioEditor/SteamAudioEditor.Build.cs
@@ -10,7 +10,6 @@ namespace UnrealBuildTool.Rules
 		{
 			PrivateIncludePaths.AddRange(
 				new string[] {
-					"SteamAudioEditor/Private",
 					"SteamAudio/Private",
 				}
 			);
diff --git a/Engine/Plugins/Runtime/Steam/SteamVR/Source/SteamVR/SteamVR.Build.cs b/Engine/Plugins/Runtime/Steam/SteamVR/Source/SteamVR/SteamVR.Build.cs
--- a/Engine/Plugins/Runtime/Steam/SteamVR/Source/SteamVR/SteamVR.Build.cs
+++ b/Engine/Plugins/Runtime/Steam/SteamVR/Source/SteamVR/SteamVR.Build.cs
@@ -6,13 +6,6 @@ namespace UnrealBuildTool.Rules
 	{
 		public SteamVR(ReadOnlyTargetRules Target) : base(Target)
 		{
-			PrivateIncludePaths.AddRange(
-				new string[] {
-					"SteamVR/Private",
-					// ... add other private include paths required here ...
-				}
-				);
-
 			PrivateDependencyModuleNames.AddRange(
 				new string[]
 				{
diff --git a/Engine/Plugins/VirtualProduction/LevelSnapshots/Source/LevelShapshotFilters/LevelSnapshotFilters.Build.cs b/Engine/Plugins/VirtualProduction/LevelSnapshots/Source/LevelShapshotFilters/LevelSnapshotFilters.Build.cs
--- a/Engine/Plugins/VirtualProduction/LevelSnapshots/Source/LevelShapshotFilters/LevelSnapshotFilters.Build.cs
+++ b/Engine/Plugins/VirtualProduction/LevelSnapshots/Source/LevelShapshotFilters/LevelSnapshotFilters.Build.cs
@@ -12,7 +12,6 @@ public class LevelSnapshotFilters : ModuleRules
 		PublicIncludePaths.AddRange(
 			new string[] 
 			{
-				Path.Combine(ModuleDirectory, "Public"),
 				Path.Combine(ModuleDirectory, "Public", "Builtin"),
 			}
 		);
diff --git a/Engine/Plugins/VirtualProduction/LevelSnapshots/Source/LevelSnapshots/LevelSnapshots.Build.cs b/Engine/Plugins/VirtualProduction/LevelSnapshots/Source/LevelSnapshots/LevelSnapshots.Build.cs
--- a/Engine/Plugins/VirtualProduction/LevelSnapshots/Source/LevelSnapshots/LevelSnapshots.Build.cs
+++ b/Engine/Plugins/VirtualProduction/LevelSnapshots/Source/LevelSnapshots/LevelSnapshots.Build.cs
@@ -12,7 +12,6 @@ public class LevelSnapshots : ModuleRules
         PublicIncludePaths.AddRange(
 	        new string[] 
 	        {
-		        Path.Combine(ModuleDirectory, "Public"),
 		        Path.Combine(ModuleDirectory, "Public", "Data"),
 		        Path.Combine(ModuleDirectory, "Public", "Restorability"),
 		        Path.Combine(ModuleDirectory, "Public", "Settings")
@@ -21,7 +20,6 @@ public class LevelSnapshots : ModuleRules
         PrivateIncludePaths.AddRange(
 	        new string[] 
 	        {
-		        Path.Combine(ModuleDirectory, "Private"),
 		        Path.Combine(ModuleDirectory, "Private", "Archive"),
 		        Path.Combine(ModuleDirectory, "Private", "Data")
 	        }
diff --git a/Engine/Plugins/VirtualProduction/Takes/Source/TakeRecorder/TakeRecorder.Build.cs b/Engine/Plugins/VirtualProduction/Takes/Source/TakeRecorder/TakeRecorder.Build.cs
--- a/Engine/Plugins/VirtualProduction/Takes/Source/TakeRecorder/TakeRecorder.Build.cs
+++ b/Engine/Plugins/VirtualProduction/Takes/Source/TakeRecorder/TakeRecorder.Build.cs
@@ -51,8 +51,6 @@ public class TakeRecorder : ModuleRules
 
 		PrivateIncludePaths.AddRange(
 			new string[] {
-				"TakeRecorder/Private",
-				"TakeRecorder/Public",
                 "TakeRecorderSources/Private",
             }
         );
diff --git a/Engine/Source/Developer/AutomationDriver/AutomationDriver.Build.cs b/Engine/Source/Developer/AutomationDriver/AutomationDriver.Build.cs
--- a/Engine/Source/Developer/AutomationDriver/AutomationDriver.Build.cs
+++ b/Engine/Source/Developer/AutomationDriver/AutomationDriver.Build.cs
@@ -6,12 +6,6 @@ public class AutomationDriver : ModuleRules
 {
     public AutomationDriver(ReadOnlyTargetRules Target) : base(Target)
     {
-        PublicIncludePaths.AddRange(
-            new string[] {
-				"Developer/AutomationDriver/Public",
-            }
-        );
-
         PrivateDependencyModuleNames.AddRange(
             new string[] {
                 "Core",
diff --git a/Engine/Source/Developer/HierarchicalLODUtilities/HierarchicalLODUtilities.Build.cs b/Engine/Source/Developer/HierarchicalLODUtilities/HierarchicalLODUtilities.Build.cs
--- a/Engine/Source/Developer/HierarchicalLODUtilities/HierarchicalLODUtilities.Build.cs
+++ b/Engine/Source/Developer/HierarchicalLODUtilities/HierarchicalLODUtilities.Build.cs
@@ -5,8 +5,6 @@ public class HierarchicalLODUtilities : ModuleRules
 {
     public HierarchicalLODUtilities(ReadOnlyTargetRules Target) : base(Target)
 	{
-        PublicIncludePaths.Add("Developer/HierarchicalLODUtilities/Public");
-
         PublicDependencyModuleNames.AddRange(
             new string[]
 			{
diff --git a/Engine/Source/Developer/MeshBoneReduction/MeshBoneReduction.Build.cs b/Engine/Source/Developer/MeshBoneReduction/MeshBoneReduction.Build.cs
--- a/Engine/Source/Developer/MeshBoneReduction/MeshBoneReduction.Build.cs
+++ b/Engine/Source/Developer/MeshBoneReduction/MeshBoneReduction.Build.cs
@@ -7,7 +7,6 @@ public class MeshBoneReduction : ModuleRules
 {
     public MeshBoneReduction(ReadOnlyTargetRules Target) : base(Target)
 	{
-		PublicIncludePaths.Add("Developer/MeshBoneReduction/Public");
 		PrivateDependencyModuleNames.Add("Core");
 		PrivateDependencyModuleNames.Add("CoreUObject");
 		PrivateDependencyModuleNames.Add("Engine");
diff --git a/Engine/Source/Editor/AIGraph/AIGraph.Build.cs b/Engine/Source/Editor/AIGraph/AIGraph.Build.cs
--- a/Engine/Source/Editor/AIGraph/AIGraph.Build.cs
+++ b/Engine/Source/Editor/AIGraph/AIGraph.Build.cs
@@ -10,7 +10,6 @@ public class AIGraph : ModuleRules
 
 		PrivateIncludePaths.AddRange(
             new string[] {
-				System.IO.Path.Combine(GetModuleDirectory("AIGraph"), "Private"),
 				System.IO.Path.Combine(GetModuleDirectory("Kismet"), "Private"),
 			}
         );
diff --git a/Engine/Source/Editor/BehaviorTreeEditor/BehaviorTreeEditor.Build.cs b/Engine/Source/Editor/BehaviorTreeEditor/BehaviorTreeEditor.Build.cs
--- a/Engine/Source/Editor/BehaviorTreeEditor/BehaviorTreeEditor.Build.cs
+++ b/Engine/Source/Editor/BehaviorTreeEditor/BehaviorTreeEditor.Build.cs
@@ -9,7 +9,6 @@ public class BehaviorTreeEditor : ModuleRules
         PrivateIncludePaths.AddRange(
             new string[] {
 				System.IO.Path.Combine(GetModuleDirectory("AIGraph"), "Private"),
-				System.IO.Path.Combine(GetModuleDirectory("BehaviorTreeEditor"), "Private"),
 			}
 		);
 
diff --git a/Engine/Source/Editor/BlueprintGraph/BlueprintGraph.Build.cs b/Engine/Source/Editor/BlueprintGraph/BlueprintGraph.Build.cs
--- a/Engine/Source/Editor/BlueprintGraph/BlueprintGraph.Build.cs
+++ b/Engine/Source/Editor/BlueprintGraph/BlueprintGraph.Build.cs
@@ -9,7 +9,6 @@ public class BlueprintGraph : ModuleRules
 	{
 		PrivateIncludePaths.AddRange(
             new string[] {
-                "Editor/BlueprintGraph/Private",
                 "Editor/KismetCompiler/Public",
             }
 		);
diff --git a/Engine/Source/Editor/ConfigEditor/ConfigEditor.Build.cs b/Engine/Source/Editor/ConfigEditor/ConfigEditor.Build.cs
--- a/Engine/Source/Editor/ConfigEditor/ConfigEditor.Build.cs
+++ b/Engine/Source/Editor/ConfigEditor/ConfigEditor.Build.cs
@@ -8,7 +8,6 @@ public class ConfigEditor : ModuleRules
 	{
 		PublicIncludePaths.AddRange(
 			new string[] {
-				"Editor/ConfigEditor/Public",
 				"Editor/ConfigEditor/Public/PropertyVisualization",
 			}
 		);
diff --git a/Engine/Source/Editor/GameplayTasksEditor/GameplayTasksEditor.Build.cs b/Engine/Source/Editor/GameplayTasksEditor/GameplayTasksEditor.Build.cs
--- a/Engine/Source/Editor/GameplayTasksEditor/GameplayTasksEditor.Build.cs
+++ b/Engine/Source/Editor/GameplayTasksEditor/GameplayTasksEditor.Build.cs
@@ -11,7 +11,6 @@ namespace UnrealBuildTool.Rules
 			PrivateIncludePaths.AddRange(
                 new string[] {
 					System.IO.Path.Combine(GetModuleDirectory("AssetTools"), "Private"),
-					System.IO.Path.Combine(GetModuleDirectory("GameplayTasksEditor"), "Private"),
 					System.IO.Path.Combine(GetModuleDirectory("Kismet"), "Private"),
 				}
 			);
diff --git a/Engine/Source/Editor/Sequencer/Sequencer.Build.cs b/Engine/Source/Editor/Sequencer/Sequencer.Build.cs
--- a/Engine/Source/Editor/Sequencer/Sequencer.Build.cs
+++ b/Engine/Source/Editor/Sequencer/Sequencer.Build.cs
@@ -8,7 +8,6 @@ public class Sequencer : ModuleRules
 	{
 		PrivateIncludePaths.AddRange(
 			new string[] {
-				"Editor/Sequencer/Private",
 				"Editor/UnrealEd/Private" // TODO: Fix this, for now it's needed for the fbx exporter
 				}
 			);
diff --git a/Engine/Source/Editor/SwarmInterface/SwarmInterface.Build.cs b/Engine/Source/Editor/SwarmInterface/SwarmInterface.Build.cs
--- a/Engine/Source/Editor/SwarmInterface/SwarmInterface.Build.cs
+++ b/Engine/Source/Editor/SwarmInterface/SwarmInterface.Build.cs
@@ -7,12 +7,6 @@ public class SwarmInterface : ModuleRules
 {
 	public SwarmInterface(ReadOnlyTargetRules Target) : base(Target)
 	{
-		PublicIncludePaths.AddRange(
-			new string[] {
-				"Editor/SwarmInterface/Public"
-			}
-		);
-
 		PublicDependencyModuleNames.AddRange(
 			new string[] {
 				"Core",
diff --git a/Engine/Source/Runtime/Experimental/Chaos/Chaos.Build.cs b/Engine/Source/Runtime/Experimental/Chaos/Chaos.Build.cs
--- a/Engine/Source/Runtime/Experimental/Chaos/Chaos.Build.cs
+++ b/Engine/Source/Runtime/Experimental/Chaos/Chaos.Build.cs
@@ -11,8 +11,6 @@ namespace UnrealBuildTool.Rules
         {
 			NumIncludedBytesPerUnityCPPOverride = 192 * 1024; // This is half of the default unity size(NumIncludedBytesPerUnityCPP) specified in TargetRules.cs
 
-			PublicIncludePaths.Add("Runtime/Experimental/Chaos/Public");
-
             PublicDependencyModuleNames.AddRange(
                 new string[] {
                 "Core",
diff --git a/Engine/Source/Runtime/Experimental/ChaosCore/ChaosCore.Build.cs b/Engine/Source/Runtime/Experimental/ChaosCore/ChaosCore.Build.cs
--- a/Engine/Source/Runtime/Experimental/ChaosCore/ChaosCore.Build.cs
+++ b/Engine/Source/Runtime/Experimental/ChaosCore/ChaosCore.Build.cs
@@ -6,8 +6,6 @@ namespace UnrealBuildTool.Rules
     {
         public ChaosCore(ReadOnlyTargetRules Target) : base(Target)
         {
-            PublicIncludePaths.Add("Runtime/Experimental/ChaosCore/Public");
-
             PublicDependencyModuleNames.AddRange(
                 new string[] {
                 "Core",
diff --git a/Engine/Source/Runtime/EyeTracker/EyeTracker.Build.cs b/Engine/Source/Runtime/EyeTracker/EyeTracker.Build.cs
--- a/Engine/Source/Runtime/EyeTracker/EyeTracker.Build.cs
+++ b/Engine/Source/Runtime/EyeTracker/EyeTracker.Build.cs
@@ -6,12 +6,6 @@ public class EyeTracker : ModuleRules
 {
     public EyeTracker(ReadOnlyTargetRules Target) : base(Target)
 	{
-		PublicIncludePaths.AddRange(
-			new string[] {
-				"Runtime/EyeTracker/Public"
-			}
-		);
-
 		PrivateDependencyModuleNames.AddRange(
 			new string[] {
 				"Core",
diff --git a/Engine/Source/Runtime/InputDevice/InputDevice.Build.cs b/Engine/Source/Runtime/InputDevice/InputDevice.Build.cs
--- a/Engine/Source/Runtime/InputDevice/InputDevice.Build.cs
+++ b/Engine/Source/Runtime/InputDevice/InputDevice.Build.cs
@@ -6,12 +6,6 @@ public class InputDevice : ModuleRules
 {
     public InputDevice(ReadOnlyTargetRules Target) : base(Target)
 	{
-		PublicIncludePaths.AddRange(
-			new string[] {
-				"Runtime/InputDevice/Public"
-			}
-			);
-
 		PrivateDependencyModuleNames.AddRange( new string[] { "Core", "CoreUObject", "Engine" } );
 	}
 }
diff --git a/Engine/Source/Runtime/Messaging/Messaging.Build.cs b/Engine/Source/Runtime/Messaging/Messaging.Build.cs
--- a/Engine/Source/Runtime/Messaging/Messaging.Build.cs
+++ b/Engine/Source/Runtime/Messaging/Messaging.Build.cs
@@ -15,7 +15,6 @@ namespace UnrealBuildTool.Rules
 
 			PrivateIncludePaths.AddRange(
 				new string[] {
-					"Runtime/Messaging/Private",
 					"Runtime/Messaging/Private/Bus",
 					"Runtime/Messaging/Private/Bridge",
 				});
diff --git a/Engine/Source/Runtime/NetworkFileSystem/NetworkFileSystem.Build.cs b/Engine/Source/Runtime/NetworkFileSystem/NetworkFileSystem.Build.cs
--- a/Engine/Source/Runtime/NetworkFileSystem/NetworkFileSystem.Build.cs
+++ b/Engine/Source/Runtime/NetworkFileSystem/NetworkFileSystem.Build.cs
@@ -19,11 +19,6 @@ namespace UnrealBuildTool.Rules
 					"CookOnTheFlyNetServer"
 				});
 
-			PublicIncludePaths.AddRange(
-				new string[] {
-					"Runtime/NetworkFileSystem/Public",
-				});
-
 			PublicDependencyModuleNames.AddRange(
 				new string[] {
 					"Core",
diff --git a/Engine/Source/Runtime/NetworkReplayStreaming/NetworkReplayStreaming/NetworkReplayStreaming.Build.cs b/Engine/Source/Runtime/NetworkReplayStreaming/NetworkReplayStreaming/NetworkReplayStreaming.Build.cs
--- a/Engine/Source/Runtime/NetworkReplayStreaming/NetworkReplayStreaming/NetworkReplayStreaming.Build.cs
+++ b/Engine/Source/Runtime/NetworkReplayStreaming/NetworkReplayStreaming/NetworkReplayStreaming.Build.cs
@@ -6,8 +6,6 @@ namespace UnrealBuildTool.Rules
 	{
 		public NetworkReplayStreaming( ReadOnlyTargetRules Target ) : base(Target)
 		{
-			PrivateIncludePaths.Add( "Runtime/NetworkReplayStreaming/NetworkReplayStreaming/Private" );
-
 			PublicDependencyModuleNames.AddRange(
 				new string[]
 				{
diff --git a/Engine/Source/Runtime/NetworkReplayStreaming/NullNetworkReplayStreaming/NullNetworkReplayStreaming.build.cs b/Engine/Source/Runtime/NetworkReplayStreaming/NullNetworkReplayStreaming/NullNetworkReplayStreaming.build.cs
--- a/Engine/Source/Runtime/NetworkReplayStreaming/NullNetworkReplayStreaming/NullNetworkReplayStreaming.build.cs
+++ b/Engine/Source/Runtime/NetworkReplayStreaming/NullNetworkReplayStreaming/NullNetworkReplayStreaming.build.cs
@@ -8,8 +8,6 @@ namespace UnrealBuildTool.Rules
 		{
             ShortName = "NullReplayStreaming";
 
-            PrivateIncludePaths.Add( "Runtime/NetworkReplayStreaming/NullNetworkReplayStreaming/Private" );
-
 			PrivateDependencyModuleNames.AddRange(
 				new string[]
 				{
diff --git a/Engine/Source/Runtime/Overlay/Overlay.Build.cs b/Engine/Source/Runtime/Overlay/Overlay.Build.cs
--- a/Engine/Source/Runtime/Overlay/Overlay.Build.cs
+++ b/Engine/Source/Runtime/Overlay/Overlay.Build.cs
@@ -21,13 +21,6 @@ namespace UnrealBuildTool.Rules
                     "UMG",
                 }
             );
-
-            PublicIncludePaths.AddRange(
-                new string[]
-                {
-                    "Runtime/Overlay/Public",
-                }
-            );
         }
 	}
 }
\ No newline at end of file
diff --git a/Engine/Source/Runtime/PhysicsCore/PhysicsCore.Build.cs b/Engine/Source/Runtime/PhysicsCore/PhysicsCore.Build.cs
--- a/Engine/Source/Runtime/PhysicsCore/PhysicsCore.Build.cs
+++ b/Engine/Source/Runtime/PhysicsCore/PhysicsCore.Build.cs
@@ -6,8 +6,6 @@ public class PhysicsCore: ModuleRules
 {
 	public PhysicsCore(ReadOnlyTargetRules Target) : base(Target)
 	{
-        PublicIncludePaths.Add("Runtime/PhysicsCore/Public");
-
         PrivateDependencyModuleNames.AddRange(
 			new string[] {
 				"Core",
diff --git a/Engine/Source/Runtime/Portal/LauncherCheck/LauncherCheck.Build.cs b/Engine/Source/Runtime/Portal/LauncherCheck/LauncherCheck.Build.cs
--- a/Engine/Source/Runtime/Portal/LauncherCheck/LauncherCheck.Build.cs
+++ b/Engine/Source/Runtime/Portal/LauncherCheck/LauncherCheck.Build.cs
@@ -6,8 +6,6 @@ public class LauncherCheck : ModuleRules
 {
 	public LauncherCheck(ReadOnlyTargetRules Target) : base(Target)
 	{
-		PublicIncludePaths.Add("Runtime/Portal/LauncherCheck/Public");
-
 		PublicDependencyModuleNames.AddRange(
 			new string[]
 			{
diff --git a/Engine/Source/Runtime/Portal/Proxies/PortalProxies.Build.cs b/Engine/Source/Runtime/Portal/Proxies/PortalProxies.Build.cs
--- a/Engine/Source/Runtime/Portal/Proxies/PortalProxies.Build.cs
+++ b/Engine/Source/Runtime/Portal/Proxies/PortalProxies.Build.cs
@@ -38,7 +38,6 @@ namespace UnrealBuildTool.Rules
 
 			PrivateIncludePaths.AddRange(
 				new string[] {
-					"Runtime/Portal/Proxies/Private",
 					"Runtime/Portal/Proxies/Private/Account",
 					"Runtime/Portal/Proxies/Private/Application",
 					"Runtime/Portal/Proxies/Private/Package",
diff --git a/Engine/Source/Runtime/PropertyPath/PropertyPath.Build.cs b/Engine/Source/Runtime/PropertyPath/PropertyPath.Build.cs
--- a/Engine/Source/Runtime/PropertyPath/PropertyPath.Build.cs
+++ b/Engine/Source/Runtime/PropertyPath/PropertyPath.Build.cs
@@ -11,11 +11,5 @@ public class PropertyPath : ModuleRules
 			"Core",
 			"CoreUObject" 
 		});
-
-		PrivateIncludePaths.AddRange(
-			new string[] {
-				"Runtime/PropertyPath/Private"
-			}
-		);
 	}
 }
diff --git a/Engine/Source/Runtime/RawMesh/RawMesh.Build.cs b/Engine/Source/Runtime/RawMesh/RawMesh.Build.cs
--- a/Engine/Source/Runtime/RawMesh/RawMesh.Build.cs
+++ b/Engine/Source/Runtime/RawMesh/RawMesh.Build.cs
@@ -6,7 +6,6 @@ public class RawMesh : ModuleRules
 {
     public RawMesh(ReadOnlyTargetRules Target) : base(Target)
 	{
-        PublicIncludePaths.AddRange(new string[] {"Runtime/RawMesh/Public"});
         PrivateDependencyModuleNames.AddRange(new string[] { "Core", "CoreUObject" });
 	}
 }
