Fix FActorSnapshotHash::Serialize not being implemented correctly

diff --git a/Engine/Plugins/VirtualProduction/LevelSnapshots/Source/LevelSnapshots/Private/Data/ActorSnapshotHash.cpp b/Engine/Plugins/VirtualProduction/LevelSnapshots/Source/LevelSnapshots/Private/Data/ActorSnapshotHash.cpp
new file mode 100644
--- /dev/null
+++ b/Engine/Plugins/VirtualProduction/LevelSnapshots/Source/LevelSnapshots/Private/Data/ActorSnapshotHash.cpp
@@ -0,0 +1,25 @@
+// Copyright Epic Games, Inc. All Rights Reserved.
+
+#include "ActorSnapshotHash.h"
+
+#include "SnapshotCustomVersion.h"
+
+bool FActorSnapshotHash::Serialize(FArchive& Archive)
+{
+	using namespace UE::LevelSnapshots;
+	Archive.UsingCustomVersion(FSnapshotCustomVersion::GUID);
+
+	if (Archive.CustomVer(FSnapshotCustomVersion::GUID) >= FSnapshotCustomVersion::FixActorHashSerialize)
+	{
+		Archive << MicroSecondsForCrc;
+		Archive << MicroSecondsForMD5;
+		Archive << Crc32DataLength;
+		Archive << MD5DataLength;
+		Archive << Crc32;
+		Archive << MD5;
+		return true;
+	}
+
+	// Use tagged property serialization
+	return false;
+}
\ No newline at end of file
diff --git a/Engine/Plugins/VirtualProduction/LevelSnapshots/Source/LevelSnapshots/Private/Data/Util/WorldData/CompressionUtil.cpp b/Engine/Plugins/VirtualProduction/LevelSnapshots/Source/LevelSnapshots/Private/Data/Util/WorldData/CompressionUtil.cpp
--- a/Engine/Plugins/VirtualProduction/LevelSnapshots/Source/LevelSnapshots/Private/Data/Util/WorldData/CompressionUtil.cpp
+++ b/Engine/Plugins/VirtualProduction/LevelSnapshots/Source/LevelSnapshots/Private/Data/Util/WorldData/CompressionUtil.cpp
@@ -117,6 +117,8 @@ namespace UE::LevelSnapshots
 		const FCompressedBuffer CompressedData = FCompressedBuffer::Load(Ar);
 		const FSharedBuffer UncompressedData = CompressedData.Decompress();
 		FBufferReader UncompressedDataArchive(const_cast<void*>(UncompressedData.GetData()), UncompressedData.GetSize(), false, Ar.IsPersistent());
+		// Needed so the compressed data can migrate
+		UncompressedDataArchive.SetCustomVersions(Ar.GetCustomVersions());
 		UncompressedDataArchive.SetIsLoading(true);
 		Private::FExtractNamesAndObjectPathsToArraysArchive ProxyArchive(UncompressedDataArchive, MoveTemp(Names), MoveTemp(Paths));
 		Private::SerializeCompressableProperties(ProxyArchive, Data);
diff --git a/Engine/Plugins/VirtualProduction/LevelSnapshots/Source/LevelSnapshots/Public/Data/ActorSnapshotHash.h b/Engine/Plugins/VirtualProduction/LevelSnapshots/Source/LevelSnapshots/Public/Data/ActorSnapshotHash.h
--- a/Engine/Plugins/VirtualProduction/LevelSnapshots/Source/LevelSnapshots/Public/Data/ActorSnapshotHash.h
+++ b/Engine/Plugins/VirtualProduction/LevelSnapshots/Source/LevelSnapshots/Public/Data/ActorSnapshotHash.h
@@ -31,25 +31,29 @@ struct LEVELSNAPSHOTS_API FActorSnapshotHash
 	UPROPERTY()
 	int32 MD5DataLength {};
 
-	/** Crc32 hash of actor when it was snapshot. Used to check for changes without loading actor. */
+	/**
+	 * Crc32 hash of actor when it was snapshot. Used to check for changes without loading actor. Sometimes MD5 is faster.
+	 * Cost is 0.0003s on "average" actors, like world settings.
+	 */
 	UPROPERTY()
 	uint32 Crc32 {};
 
+	/**
+	 * MD5 hash of actor when it was snapshot. Used to check for changes without loading actor. Sometimes CRC32 is faster.
+	 * Cost is 0.0003s on "average" actors, like world settings. 
+	 * CRC32 is generally to be preferred but we do MD5 to future proof cases where there may be collisions. Users can configure this in project settings.
+	 */
 	FMD5Hash MD5;
 	
 	bool HasCrc32() const { return MicroSecondsForCrc > 0.0; }
-	bool HasMD5() const { return MicroSecondsForMD5 > 0.0; }
-
-	bool Serialize(FArchive& Archive)
+	bool HasMD5() const
 	{
-		Archive << MicroSecondsForCrc;
-		Archive << MicroSecondsForMD5;
-		Archive << Crc32DataLength;
-		Archive << MD5DataLength;
-		Archive << Crc32;
-		Archive << MD5;
-		return true;
+		return MicroSecondsForMD5 > 0.0
+			// Data saved before 5.2 can have 0.0 > MicroSecondsForMD5 but invalid MD5 because Serialize was not implemented correctly (see FSnapshotCustomVersion::FixActorHashSerialize)
+			&& MD5.IsValid();
 	}
+
+	bool Serialize(FArchive& Archive);
 };
 
 template<>
@@ -57,6 +61,6 @@ struct TStructOpsTypeTraits<FActorSnapshotHash> : public TStructOpsTypeTraitsBas
 {
 	enum 
 	{ 
-		WithSerialize = true
+		WithSerializer = true
 	};
 };
\ No newline at end of file
diff --git a/Engine/Plugins/VirtualProduction/LevelSnapshots/Source/LevelSnapshots/Public/Data/SnapshotCustomVersion.h b/Engine/Plugins/VirtualProduction/LevelSnapshots/Source/LevelSnapshots/Public/Data/SnapshotCustomVersion.h
--- a/Engine/Plugins/VirtualProduction/LevelSnapshots/Source/LevelSnapshots/Public/Data/SnapshotCustomVersion.h
+++ b/Engine/Plugins/VirtualProduction/LevelSnapshots/Source/LevelSnapshots/Public/Data/SnapshotCustomVersion.h
@@ -42,6 +42,9 @@ namespace UE::LevelSnapshots
 			 */
 			NonEditableComponentsAreCaptured = 7,
 
+			/** FActorSnapshotHash::Serialize was not implemented correctly (typo with TStructOpsTypeTraits using WithSerialize instead of WithSerializer). */
+			FixActorHashSerialize = 8,
+
 			// -----<new versions can be added above this line>-------------------------------------------------
 			VersionPlusOne,
 			LatestVersion = VersionPlusOne - 1
