Make level snapshots call PreEditChange and PostEditChange again. Fixes issue where SkyLight would not update the scene's light after restoring.

diff --git a/Engine/Plugins/VirtualProduction/LevelSnapshots/Source/LevelSnapshots/Private/Archive/ApplySnapshotToEditorArchive.cpp b/Engine/Plugins/VirtualProduction/LevelSnapshots/Source/LevelSnapshots/Private/Archive/ApplySnapshotToEditorArchive.cpp
--- a/Engine/Plugins/VirtualProduction/LevelSnapshots/Source/LevelSnapshots/Private/Archive/ApplySnapshotToEditorArchive.cpp
+++ b/Engine/Plugins/VirtualProduction/LevelSnapshots/Source/LevelSnapshots/Private/Archive/ApplySnapshotToEditorArchive.cpp
@@ -32,6 +32,7 @@ void UE::LevelSnapshots::Private::FApplySnapshotToEditorArchive::ApplyToExisting
 	const FApplySnapshotPropertiesScope NotifySnapshotListeners({ InOriginalObject, InSelectionMapForResolvingSubobjects, Selection, true });
 #if WITH_EDITOR
 	InOriginalObject->Modify();
+	InOriginalObject->PreEditChange(nullptr);
 #endif
 	
 	// 1. Serialize archetype first to handle the case were the archetype has changed properties since the snapshot was taken
@@ -44,6 +45,10 @@ void UE::LevelSnapshots::Private::FApplySnapshotToEditorArchive::ApplyToExisting
 	// Step 2: Serialise  properties that were different from CDO at time of snapshotting and that are still different from CDO
 	FApplySnapshotToEditorArchive ApplySavedData(InObjectData, InSharedData, InOriginalObject, InSelectionMapForResolvingSubobjects, Selection, Cache);
 	InOriginalObject->Serialize(ApplySavedData);
+
+#if WITH_EDITOR
+	InOriginalObject->PostEditChange();
+#endif
 }
 
 void UE::LevelSnapshots::Private::FApplySnapshotToEditorArchive::ApplyToEditorWorldObjectRecreatedWithArchetype(
