Level Snapshots supports renaming levels again and rehashes actors map when modified.

diff --git a/Engine/Plugins/VirtualProduction/LevelSnapshots/Source/LevelSnapshots/Private/Data/WorldSnapshotData.cpp b/Engine/Plugins/VirtualProduction/LevelSnapshots/Source/LevelSnapshots/Private/Data/WorldSnapshotData.cpp
--- a/Engine/Plugins/VirtualProduction/LevelSnapshots/Source/LevelSnapshots/Private/Data/WorldSnapshotData.cpp
+++ b/Engine/Plugins/VirtualProduction/LevelSnapshots/Source/LevelSnapshots/Private/Data/WorldSnapshotData.cpp
@@ -7,6 +7,7 @@
 #include "Util/Property/PropertyIterator.h"
 
 #include "Serialization/BufferArchive.h"
+#include "UObject/UnrealType.h"
 #include "Util/SnapshotUtil.h"
 #include "Util/WorldData/ClassDataUtil.h"
 #include "Util/WorldData/CompressionUtil.h"
@@ -30,12 +31,23 @@ namespace UE::LevelSnapshots::Private::Internal
 {
 	static void CollectActorReferences(FWorldSnapshotData& WorldData, FArchive& Ar)
 	{
+		// If the archive is modifying our references, e.g. Ar.IsModifyingWeakAndStrongReferences(), a rehash may be needed
+		bool bNeedsRehash = false;
 		for (auto ActorIt = WorldData.ActorData.CreateIterator(); ActorIt; ++ActorIt)
 		{
+			const FSoftObjectPath BeforeSerialization = ActorIt->Key;
 			Ar << ActorIt->Key;
+			bNeedsRehash |= ActorIt->Key != BeforeSerialization;
 			FSoftClassPath ClassPath = GetClass(ActorIt->Value, WorldData);
 			Ar << ClassPath;
 		}
+
+		if (bNeedsRehash)
+		{
+			const FMapProperty* MapProperty = CastField<FMapProperty>(FWorldSnapshotData::StaticStruct()->FindPropertyByName(GET_MEMBER_NAME_CHECKED(FWorldSnapshotData, ActorData)));
+			FScriptMapHelper Helper(MapProperty, &WorldData.ActorData);
+			Helper.Rehash();
+		}
 	}
 
 	static void CollectClassDefaultReferences(FWorldSnapshotData& WorldData, FArchive& Ar)
@@ -61,24 +73,11 @@ namespace UE::LevelSnapshots::Private::Internal
 		}
 	}
 	
-	static void CollectReferencesAndNames(bool bSkipActorReferences, FWorldSnapshotData& WorldData, FArchive& Ar)
+	static void CollectReferencesAndNames(FWorldSnapshotData& WorldData, FArchive& Ar)
 	{
 		// References
-		if (bSkipActorReferences)
-		{
-			for (FSoftObjectPath& Path : WorldData.SerializedObjectReferences)
-			{
-				if (!IsPathToWorldObject(Path))
-				{
-					Ar << Path;
-				}
-			}
-		}
-		else
-		{
-			Ar << WorldData.SerializedObjectReferences;
-			CollectActorReferences(WorldData, Ar);
-		}
+		Ar << WorldData.SerializedObjectReferences;
+		CollectActorReferences(WorldData, Ar);
 		CollectClassDefaultReferences(WorldData, Ar);
 
 		// Names
@@ -197,11 +196,7 @@ bool FWorldSnapshotData::Serialize(FArchive& Ar)
 	// When this struct is saved, the save algorithm collects references. It's faster if we just give it the info directly.
 	if (Ar.IsObjectReferenceCollector())
 	{
-		// Actor references must not be renamed because a snapshot is supposed to keep track of what was in the level at a given time.
-		// If an actor is renamed or moved to another level, the snapshot should ignore those changes.
-		const bool bIsRenameArchive = !Ar.IsPersistent() && Ar.IsSaving() && Ar.IsModifyingWeakAndStrongReferences();
-		const bool bSkipActorReferences = bIsRenameArchive;
-		UE::LevelSnapshots::Private::Internal::CollectReferencesAndNames(bSkipActorReferences, *this, Ar);
+		UE::LevelSnapshots::Private::Internal::CollectReferencesAndNames(*this, Ar);
 		return true;
 	}
 
