Fix foliage not restoring anymore in Level Snapshots and prevent if data is from 5.1. Warn if data is from before 5.1 and allow only with console variable (separate possible crash found)

diff --git a/Engine/Plugins/VirtualProduction/LevelSnapshots/Source/FoliageSupport/Private/FoliageLevelSnapshotsConsoleVariables.cpp b/Engine/Plugins/VirtualProduction/LevelSnapshots/Source/FoliageSupport/Private/FoliageLevelSnapshotsConsoleVariables.cpp
new file mode 100644
--- /dev/null
+++ b/Engine/Plugins/VirtualProduction/LevelSnapshots/Source/FoliageSupport/Private/FoliageLevelSnapshotsConsoleVariables.cpp
@@ -0,0 +1,14 @@
+// Copyright Epic Games, Inc. All Rights Reserved.
+
+#include "FoliageLevelSnapshotsConsoleVariables.h"
+
+namespace UE::LevelSnapshots::Foliage::Private
+{
+	TAutoConsoleVariable<bool> CVarAllowFoliageDataPre5dot1(
+		TEXT("LevelSnapshots.AllowFoliageDataBefore5dot1"),
+		false,
+		TEXT("Foliage data saved before 5.1 was saved in a bad format which may cause a crash. Data from 5.1 is unusable. It is safe to enable this if either:\n"
+			"A. you added static mesh actors into the foliage tool (e.g. drop-drop static mesh from content browser)\n"
+			"B. you added foliage types into the foliage tool (e.g. drag-dop foliage type from content browser)")
+		);
+}
\ No newline at end of file
diff --git a/Engine/Plugins/VirtualProduction/LevelSnapshots/Source/FoliageSupport/Private/FoliageLevelSnapshotsConsoleVariables.h b/Engine/Plugins/VirtualProduction/LevelSnapshots/Source/FoliageSupport/Private/FoliageLevelSnapshotsConsoleVariables.h
new file mode 100644
--- /dev/null
+++ b/Engine/Plugins/VirtualProduction/LevelSnapshots/Source/FoliageSupport/Private/FoliageLevelSnapshotsConsoleVariables.h
@@ -0,0 +1,17 @@
+// Copyright Epic Games, Inc. All Rights Reserved.
+
+#pragma once
+
+#include "HAL/IConsoleManager.h"
+
+namespace UE::LevelSnapshots::Foliage::Private
+{
+	/**
+	 * Foliage data saved before 5.2 was saved in a bad format which may cause a crash. It is safe to enable this if either:
+	 *	1. All foliage were assets (possible / conceivable)
+	 *	2. All foliage were subobjects (possible / conceivable)
+	 *	3. First assets were added, then subobjects subobjects (unlikely)
+	 *	Subobject = add by drag-droping static mesh, asset = drag-drop UFoliageType asset
+	 */
+	extern TAutoConsoleVariable<bool> CVarAllowFoliageDataPre5dot1;
+}
diff --git a/Engine/Plugins/VirtualProduction/LevelSnapshots/Source/FoliageSupport/Private/FoliageSupport/Apply/ApplyDataUtils.cpp b/Engine/Plugins/VirtualProduction/LevelSnapshots/Source/FoliageSupport/Private/FoliageSupport/Apply/ApplyDataUtils.cpp
new file mode 100644
--- /dev/null
+++ b/Engine/Plugins/VirtualProduction/LevelSnapshots/Source/FoliageSupport/Private/FoliageSupport/Apply/ApplyDataUtils.cpp
@@ -0,0 +1,70 @@
+// Copyright Epic Games, Inc. All Rights Reserved.
+
+#include "ApplyDataUtils.h"
+
+#include "LevelSnapshotsLog.h"
+#include "Selection/PropertySelectionMap.h"
+
+#include "FoliageType.h"
+#include "InstancedFoliageActor.h"
+
+namespace UE::LevelSnapshots::Foliage::Private
+{
+	FFoliageInfo* FindOrAddFoliageInfo(UFoliageType* FoliageType, AInstancedFoliageActor* FoliageActor)
+	{
+		FFoliageInfo* FoliageInfo = FoliageActor->FindInfo(FoliageType);
+		if (!FoliageInfo)
+		{
+			FoliageInfo = &FoliageActor->AddFoliageInfo(FoliageType).Get();
+		}
+
+		UE_CLOG(FoliageInfo == nullptr, LogLevelSnapshots, Warning, TEXT("Failed to create foliage info for foliage type %s"), *FoliageType->GetPathName());
+		return FoliageInfo; 
+	}
+
+	void HandleExistingFoliageUsingRequiredComponent(AInstancedFoliageActor* FoliageActor, FName RequiredComponentName, const TMap<FName, UFoliageType*>& PreexistingComponentToFoliageType, UFoliageType* AllowedFoliageType)
+	{
+		// Handles the following case:
+		// 1. Add (static mesh) foliage type FoliageA. Let's suppose the component that is added is called Component01.
+		// 2. Take snapshot
+		// 3. Reload map without saving
+		// 4. Add foliage type FoliageB. This creates a new component that is also called Component01.
+		// 5. Restore snapshot
+		// 
+		// Without the below check, both FoliageA and FoliageB would use Component01 to add instances.
+		// We remove the foliage type under the "interpretation" that Component01 is "restored" to reuse the previous foliage type.
+		if (UFoliageType* const* FoliageType = PreexistingComponentToFoliageType.Find(RequiredComponentName); FoliageType && *FoliageType != AllowedFoliageType)
+		{
+			TUniqueObj<FFoliageInfo> FoliageInfo;
+			FoliageActor->RemoveFoliageInfoAndCopyValue(*FoliageType, FoliageInfo);
+		}
+	}
+
+	void RemoveComponentAutoRecreatedByLevelSnapshots(AInstancedFoliageActor* FoliageActor, TOptional<FName> RecreatedComponentName)
+	{
+		if (ensure(RecreatedComponentName))
+		{
+			for (UHierarchicalInstancedStaticMeshComponent* Component : TInlineComponentArray<UHierarchicalInstancedStaticMeshComponent*>(FoliageActor))
+			{
+				if (Component->GetFName() == *RecreatedComponentName)
+				{
+					Component->DestroyComponent();
+					return;
+				}
+			}
+		}
+	}
+
+	TMap<FName, UFoliageType*> BuildComponentToFoliageType(AInstancedFoliageActor* FoliageActor)
+	{
+		TMap<FName, UFoliageType*> Result;
+		for (auto It = FoliageActor->GetFoliageInfos().CreateConstIterator(); It; ++It)
+		{
+			if (UHierarchicalInstancedStaticMeshComponent* Component = It->Value->GetComponent())
+			{
+				Result.Add(Component->GetFName(), It->Key);
+			}
+		}
+		return Result;
+	}
+}
\ No newline at end of file
diff --git a/Engine/Plugins/VirtualProduction/LevelSnapshots/Source/FoliageSupport/Private/FoliageSupport/Apply/ApplyDataUtils.h b/Engine/Plugins/VirtualProduction/LevelSnapshots/Source/FoliageSupport/Private/FoliageSupport/Apply/ApplyDataUtils.h
new file mode 100644
--- /dev/null
+++ b/Engine/Plugins/VirtualProduction/LevelSnapshots/Source/FoliageSupport/Private/FoliageSupport/Apply/ApplyDataUtils.h
@@ -0,0 +1,38 @@
+// Copyright Epic Games, Inc. All Rights Reserved.
+
+#pragma once
+
+#include "Containers/Map.h"
+#include "Misc/Optional.h"
+#include "UObject/NameTypes.h"
+
+class AInstancedFoliageActor;
+class UFoliageType;
+struct FFoliageInfo;
+
+namespace UE::LevelSnapshots::Foliage::Private
+{
+	/** Adds the foliage type to the instanced actor or returns the associated FFoliageInfo if already present. */
+	FFoliageInfo* FindOrAddFoliageInfo(UFoliageType* FoliageType, AInstancedFoliageActor* FoliageActor);
+
+	/**
+	 * Handles name clashes with pre-existing instanced components.
+	 *
+	 * This handles in particular this case
+	 * 1. Add (static mesh) foliage type FoliageA. Let's suppose the component that is added is called Component01.
+	 * 2. Take snapshot
+	 * 3. Reload map without saving
+	 * 4. Add foliage type FoliageB. This creates a new component that is also called Component01.
+	 * 5. Restore snapshot
+	 *
+	 * Both FoliageA and FoliageB would use Component01 to add instances.
+	 * We remove the foliage type under the "interpretation" that Component01 is "restored" to reuse the previous foliage type
+	 */
+	void HandleExistingFoliageUsingRequiredComponent(AInstancedFoliageActor* FoliageActor, FName RequiredComponentName, const TMap<FName, UFoliageType*>& PreexistingComponentToFoliageType, UFoliageType* AllowedFoliageType);
+
+	/** Default LS restoration will recreate all components. However, if the user said to skip restoring some of the components, we'll remove the wrongly added components again. */
+	void RemoveComponentAutoRecreatedByLevelSnapshots(AInstancedFoliageActor* FoliageActor, TOptional<FName> RecreatedComponentName);
+
+	/** Returns an inverse map binding hierarchical mesh component names to the foliage type they represent. */
+	TMap<FName, UFoliageType*> BuildComponentToFoliageType(AInstancedFoliageActor* FoliageActor);
+}
diff --git a/Engine/Plugins/VirtualProduction/LevelSnapshots/Source/FoliageSupport/Private/FoliageSupport/Apply/ApplyData_5dot2.cpp b/Engine/Plugins/VirtualProduction/LevelSnapshots/Source/FoliageSupport/Private/FoliageSupport/Apply/ApplyData_5dot2.cpp
new file mode 100644
--- /dev/null
+++ b/Engine/Plugins/VirtualProduction/LevelSnapshots/Source/FoliageSupport/Private/FoliageSupport/Apply/ApplyData_5dot2.cpp
@@ -0,0 +1,117 @@
+// Copyright Epic Games, Inc. All Rights Reserved.
+
+#include "FoliageSupport/Data/InstancedFoliageActorData.h"
+
+#include "ApplyDataUtils.h"
+#include "FoliageRestorationInfo.h"
+#include "LevelSnapshotsLog.h"
+#include "Selection/PropertySelectionMap.h"
+#include "SnapshotCustomVersion.h"
+
+#include "FoliageType.h"
+#include "InstancedFoliageActor.h"
+#include "Serialization/Archive.h"
+#include "Misc/ScopeExit.h"
+
+namespace UE::LevelSnapshots::Foliage::Private
+{
+	static void ApplyAssets_5dot2(
+		FArchive& Archive,
+		AInstancedFoliageActor* FoliageActor,
+		const FFoliageRestorationInfo& RestorationInfo,
+		const TMap<TSoftObjectPtr<UFoliageType>, FFoliageInfoData>& FoliageAssets,
+		const TMap<FName, UFoliageType*> PreexistingComponentToFoliageType)
+	{
+		for (auto AssetIt = FoliageAssets.CreateConstIterator(); AssetIt; ++AssetIt)
+		{
+			const FFoliageInfoData& FoliageInfoData = AssetIt->Value;
+			UFoliageType* FoliageType = AssetIt->Key.LoadSynchronous();
+			if (!FoliageType)
+			{
+				UE_LOG(LogLevelSnapshots, Warning, TEXT("Foliage type %s is missing. Component %s will not be restored."), *AssetIt->Key.ToString(), *FoliageInfoData.GetComponentName().Get(NAME_None).ToString());
+				continue;
+			}
+
+			if (RestorationInfo.ShouldSkipFoliageType(FoliageInfoData) || !RestorationInfo.ShouldSerializeFoliageType(FoliageInfoData))
+			{
+				continue;
+			}
+
+			FFoliageInfo* FoliageInfo = FindOrAddFoliageInfo(FoliageType, FoliageActor);
+			if (FoliageInfo)
+			{
+				// If there were no instances, no component existed, hence no component name could be saved
+				const bool bHadAtLeastOneInstanceWhenSaved = FoliageInfoData.GetComponentName().IsSet();
+				if (bHadAtLeastOneInstanceWhenSaved)
+				{
+					HandleExistingFoliageUsingRequiredComponent(FoliageActor, *FoliageInfoData.GetComponentName(), PreexistingComponentToFoliageType, FoliageType);
+				}
+				
+				FoliageInfoData.ApplyTo(Archive, *FoliageInfo);
+			}
+			else
+			{
+				// Standard component restoration has already recreated the component. Remove it again.
+				RemoveComponentAutoRecreatedByLevelSnapshots(FoliageActor, FoliageInfoData.GetComponentName());
+			}
+		}
+	}
+
+	static void ApplySubobjects_5dot2(
+		FArchive& Archive,
+		AInstancedFoliageActor* FoliageActor,
+		const FFoliageRestorationInfo& RestorationInfo,
+		const TArray<FSubobjectFoliageInfoData>& SubobjectData, 
+		const TMap<FName, UFoliageType*> PreexistingComponentToFoliageType)
+	{
+		for (const FSubobjectFoliageInfoData& InstanceData : SubobjectData)
+		{
+			if (RestorationInfo.ShouldSkipFoliageType(InstanceData) || !RestorationInfo.ShouldSerializeFoliageType(InstanceData))
+			{
+				continue;
+			}
+
+			UFoliageType* FoliageType = InstanceData.FindOrRecreateSubobject(FoliageActor);
+			if (!FoliageType)
+			{
+				UE_LOG(LogLevelSnapshots, Error, TEXT("Failed to recreate foliage type. Skipping..."));
+				continue;
+			}
+
+			bool bWasRestored = false;
+			ON_SCOPE_EXIT
+			{
+				if (!bWasRestored)
+				{
+					// Standard component restoration has already recreated the component. Remove it again.
+					RemoveComponentAutoRecreatedByLevelSnapshots(FoliageActor, InstanceData.GetComponentName());
+				}
+			};
+
+			if (!InstanceData.ApplyToFoliageType(Archive, FoliageType))
+			{
+				UE_LOG(LogLevelSnapshots, Warning, TEXT("Skipping restoration of foliage item due to errors."));
+				continue;
+			}
+				
+			FFoliageInfo* FoliageInfo = FindOrAddFoliageInfo(FoliageType, FoliageActor);
+			if (FoliageInfo
+				&& ensureAlwaysMsgf(InstanceData.GetComponentName(), TEXT("ComponentName is supposed to have been saved. Investigate. Maybe you used an unsupported foliage type (only EFoliageImplType::StaticMesh is supported)?")))
+			{
+				HandleExistingFoliageUsingRequiredComponent(FoliageActor, *InstanceData.GetComponentName(), PreexistingComponentToFoliageType, FoliageType);
+				InstanceData.ApplyToFoliageInfo(Archive, *FoliageInfo);
+				bWasRestored = true;
+			}
+		}
+	}
+	
+	void FInstancedFoliageActorData::ApplyTo_5dot2(FArchive& Archive, AInstancedFoliageActor* FoliageActor, const FPropertySelectionMap& SelectedProperties, bool bWasRecreated) const
+	{
+		check(Archive.CustomVer(FSnapshotCustomVersion::GUID) >= FSnapshotCustomVersion::FoliageTypesUnreadable);
+
+		const FFoliageRestorationInfo RestorationInfo = FFoliageRestorationInfo::From(FoliageActor, SelectedProperties, bWasRecreated);
+		const TMap<FName, UFoliageType*> PreexistingComponentToFoliageType = BuildComponentToFoliageType(FoliageActor);
+		ApplyAssets_5dot2(Archive, FoliageActor, RestorationInfo, FoliageAssets, PreexistingComponentToFoliageType);
+		ApplySubobjects_5dot2(Archive, FoliageActor, RestorationInfo, SubobjectData, PreexistingComponentToFoliageType);
+	}
+}
\ No newline at end of file
diff --git a/Engine/Plugins/VirtualProduction/LevelSnapshots/Source/FoliageSupport/Private/FoliageSupport/Apply/ApplyData_Pre5dot1.cpp b/Engine/Plugins/VirtualProduction/LevelSnapshots/Source/FoliageSupport/Private/FoliageSupport/Apply/ApplyData_Pre5dot1.cpp
new file mode 100644
--- /dev/null
+++ b/Engine/Plugins/VirtualProduction/LevelSnapshots/Source/FoliageSupport/Private/FoliageSupport/Apply/ApplyData_Pre5dot1.cpp
@@ -0,0 +1,172 @@
+// Copyright Epic Games, Inc. All Rights Reserved.
+
+#include "FoliageSupport/Data/InstancedFoliageActorData.h"
+
+#include "ApplyDataUtils.h"
+#include "FoliageLevelSnapshotsConsoleVariables.h"
+#include "FoliageRestorationInfo.h"
+#include "LevelSnapshotsLog.h"
+#include "Selection/PropertySelectionMap.h"
+#include "SnapshotCustomVersion.h"
+
+#include "FoliageType.h"
+#include "InstancedFoliageActor.h"
+#include "Misc/ScopeExit.h"
+
+namespace UE::LevelSnapshots::Foliage::Private
+{
+	static void ApplyAssets_Pre5dot2(
+		FArchive& Archive,
+		AInstancedFoliageActor* FoliageActor,
+		const FFoliageRestorationInfo& RestorationInfo,
+		const TMap<TSoftObjectPtr<UFoliageType>, FFoliageInfoData>& FoliageAssets,
+		const TMap<FName, UFoliageType*> PreexistingComponentToFoliageType)
+	{
+		// Some objects may not want to be restored (either due to errors or because user deselected them) ...
+		int64 DataToSkip = 0;
+		
+		for (auto AssetIt = FoliageAssets.CreateConstIterator(); AssetIt; ++AssetIt)
+		{
+			const FFoliageInfoData& FoliageInfoData = AssetIt->Value;
+			
+			UFoliageType* FoliageType = AssetIt->Key.LoadSynchronous();
+			if (!FoliageType)
+			{
+				UE_LOG(LogLevelSnapshots, Warning, TEXT("Foliage type %s is missing. Component %s will not be restored."), *AssetIt->Key.ToString(), *FoliageInfoData.GetComponentName().Get(NAME_None).ToString());
+				continue;
+			}
+
+			if (RestorationInfo.ShouldSkipFoliageType(FoliageInfoData) || !RestorationInfo.ShouldSerializeFoliageType(FoliageInfoData))
+			{
+				continue;
+			}
+
+			FFoliageInfo* FoliageInfo = FindOrAddFoliageInfo(FoliageType, FoliageActor);
+			if (FoliageInfo)
+			{
+				// If there were no instances, no component existed, hence no component name could be saved
+				const bool bHadAtLeastOneInstanceWhenSaved = FoliageInfoData.GetComponentName().IsSet();
+				if (bHadAtLeastOneInstanceWhenSaved)
+				{
+					HandleExistingFoliageUsingRequiredComponent(FoliageActor, *FoliageInfoData.GetComponentName(), PreexistingComponentToFoliageType, FoliageType);
+				}
+
+				// ... since Archive data is sequential we must account for possibly having skipped some data
+				Archive.Seek(Archive.Tell() + DataToSkip);
+				DataToSkip = 0;
+				
+				FoliageInfoData.ApplyTo(Archive, *FoliageInfo);
+			}
+			else
+			{
+				// Standard component restoration has already recreated the component. Remove it again.
+				RemoveComponentAutoRecreatedByLevelSnapshots(FoliageActor, FoliageInfoData.GetComponentName());
+				DataToSkip += FoliageInfoData.GetFoliageInfoArchiveSize();
+			}
+		}
+	}
+
+	static void ApplySubobjects_Pre5dot2(
+		FArchive& Archive,
+		AInstancedFoliageActor* FoliageActor,
+		const FFoliageRestorationInfo& RestorationInfo,
+		const TArray<FSubobjectFoliageInfoData>& SubobjectData, 
+		const TMap<FName, UFoliageType*> PreexistingComponentToFoliageType)
+	{
+		// Some objects may not want to be restored (either due to errors or because user deselected them) ...
+		int64 DataToSkip = 0;
+		
+		for (const FSubobjectFoliageInfoData& InstanceData : SubobjectData)
+		{
+			if (RestorationInfo.ShouldSkipFoliageType(InstanceData) || !RestorationInfo.ShouldSerializeFoliageType(InstanceData))
+			{
+				continue;
+			}
+
+			UFoliageType* FoliageType = InstanceData.FindOrRecreateSubobject(FoliageActor);
+			if (!FoliageType)
+			{
+				UE_LOG(LogLevelSnapshots, Error, TEXT("Failed to recreate foliage type. Skipping..."));
+				continue;
+			}
+
+			// ... since Archive data is sequential we must account for possibly having skipped some data
+			Archive.Seek(Archive.Tell() + DataToSkip);
+			DataToSkip = 0;
+			
+			bool bWasRestored = false;
+			bool bWasPartiallyRestored = false;
+			ON_SCOPE_EXIT
+			{
+				if (!bWasRestored)
+				{
+					// Standard component restoration has already recreated the component. Remove it again.
+					RemoveComponentAutoRecreatedByLevelSnapshots(FoliageActor, InstanceData.GetComponentName());
+					DataToSkip += static_cast<uint64>(bWasPartiallyRestored) * InstanceData.GetFoliageTypeArchiveSize() + InstanceData.GetFoliageInfoArchiveSize();
+					// Technically we have to remove the foliage info from FoliageActor ... it was already serialized into FoliageInfos
+				}
+			};
+
+			bWasPartiallyRestored = true;
+			if (!InstanceData.ApplyToFoliageType(Archive, FoliageType))
+			{
+				UE_LOG(LogLevelSnapshots, Warning, TEXT("Skipping restoration of foliage item due to errors."));
+				continue;
+			}
+				
+			FFoliageInfo* FoliageInfo = FindOrAddFoliageInfo(FoliageType, FoliageActor);
+			if (FoliageInfo
+				&& ensureAlwaysMsgf(InstanceData.GetComponentName(), TEXT("ComponentName is supposed to have been saved. Investigate. Maybe you used an unsupported foliage type (only EFoliageImplType::StaticMesh is supported)?")))
+			{
+				HandleExistingFoliageUsingRequiredComponent(FoliageActor, *InstanceData.GetComponentName(), PreexistingComponentToFoliageType, FoliageType);
+				InstanceData.ApplyToFoliageInfo(Archive, *FoliageInfo);
+				bWasRestored = true;
+			}
+		}
+	}
+	
+	void FInstancedFoliageActorData::ApplyData_Pre5dot2(FArchive& Archive, AInstancedFoliageActor* FoliageActor, const FPropertySelectionMap& SelectedProperties, bool bWasRecreated) const
+	{
+		check(Archive.CustomVer(FSnapshotCustomVersion::GUID) < FSnapshotCustomVersion::FoliageTypesUnreadable);
+		
+		/*
+		 * Remember: Subobject = add by drag-droping static mesh, asset = drag-drop UFoliageType asset
+		 *
+		 * Simplified version of the previous saving code:
+		 * FoliageActor->ForEachFoliageInfo([this, &Archive, FoliageActor](UFoliageType* FoliageType, FFoliageInfo& FoliageInfo)
+		 * {
+		 *		if (bIsSubobject) Internal::SaveSubobject(Archive, FoliageType, FoliageInfo, SubobjectData);
+		 *		else Internal::SaveAsset(Archive, FoliageType, FoliageInfo, FoliageAssets);
+		 * });
+		 * 
+		 * ApplyAssets and ApplySubobjects just iterate the map from beginning to end.
+		 * If the foliage infos are interleaved, we'll get a crash. There is NO way to know in advance.
+		 * 
+		 * Example interleaved, dangerous saved state (with hypothetical byte sizes):
+		 * [0] Subobject1	(42 bytes)
+		 * [1] Subobject2	(58 bytes)
+		 * [2] Asset1		(100 bytes)		// Ok so far
+		 * [3] Subobject3	(50 bytes)		// This will cause a crash
+		 *
+		 * The above would crash in ApplyAssets because it would do Archive << x, reading data it expects to be Asset1.
+		 * This code only works in the following cases:
+		 * 1. Everything assets (possible / conceivable)
+		 * 2. Everything subobjects (possible / conceivable)
+		 * 3. First assets, then subobjects (unlikely)
+		 */
+		const bool bIsSafeData = Archive.CustomVer(FSnapshotCustomVersion::GUID) < FSnapshotCustomVersion::CustomSubobjectSoftObjectPathRefactor;
+		if (bIsSafeData || ensureMsgf(CVarAllowFoliageDataPre5dot1.GetValueOnGameThread(), TEXT("This check was already supposed to have happened in FFoliageSupport's IActorSnapshotFilter implementation")))
+		{
+			UE_CLOG(!bIsSafeData, LogLevelSnapshots, Warning, TEXT("Data saved before 5.1 may crash the editor (unlikely)."));
+			
+			const FFoliageRestorationInfo RestorationInfo = FFoliageRestorationInfo::From(FoliageActor, SelectedProperties, bWasRecreated);
+			const TMap<FName, UFoliageType*> PreexistingComponentToFoliageType = BuildComponentToFoliageType(FoliageActor);
+			ApplyAssets_Pre5dot2(Archive, FoliageActor, RestorationInfo, FoliageAssets, PreexistingComponentToFoliageType);
+			ApplySubobjects_Pre5dot2(Archive, FoliageActor, RestorationInfo, SubobjectData, PreexistingComponentToFoliageType);
+		}
+		else
+		{
+			UE_LOG(LogLevelSnapshots, Warning, TEXT("Use \"LevelSnapshots.DangerousAllowFoliageDataPre5dot1 true\" to enable restoring foliage data from before 5.2 (may crash)."));
+		}
+	}
+}
diff --git a/Engine/Plugins/VirtualProduction/LevelSnapshots/Source/FoliageSupport/Private/FoliageSupport/Apply/FoliageRestorationInfo.cpp b/Engine/Plugins/VirtualProduction/LevelSnapshots/Source/FoliageSupport/Private/FoliageSupport/Apply/FoliageRestorationInfo.cpp
new file mode 100644
--- /dev/null
+++ b/Engine/Plugins/VirtualProduction/LevelSnapshots/Source/FoliageSupport/Private/FoliageSupport/Apply/FoliageRestorationInfo.cpp
@@ -0,0 +1,67 @@
+﻿// Copyright Epic Games, Inc. All Rights Reserved.
+
+#include "FoliageSupport/Apply/FoliageRestorationInfo.h"
+
+#include "FoliageSupport/Data/FoliageInfoData.h"
+#include "LevelSnapshotsLog.h"
+#include "Selection/AddedAndRemovedComponentInfo.h"
+#include "Selection/RestorableObjectSelection.h"
+#include "Selection/PropertySelectionMap.h"
+
+#include "InstancedFoliageActor.h"
+
+UE::LevelSnapshots::Foliage::Private::FFoliageRestorationInfo UE::LevelSnapshots::Foliage::Private::FFoliageRestorationInfo::From(AInstancedFoliageActor* Object, const FPropertySelectionMap& SelectionMap, bool bWasRecreated)
+{
+	FFoliageRestorationInfo Result;
+	Result.bWasRecreated = bWasRecreated;
+	
+	const FRestorableObjectSelection ObjectSelection = SelectionMap.GetObjectSelection(Object);
+	if (const FAddedAndRemovedComponentInfo* ComponentInfo = ObjectSelection.GetComponentSelection())
+	{
+		Result.EditorWorldComponentsToRemove = ComponentInfo->EditorWorldComponentsToRemove.Array();
+		Result.SnapshotComponentsToAdd = ComponentInfo->SnapshotComponentsToAdd.Array();
+	}
+
+	for (UHierarchicalInstancedStaticMeshComponent* FoliageComp : TInlineComponentArray<UHierarchicalInstancedStaticMeshComponent*>(Object))
+	{
+		if (SelectionMap.HasChanges(FoliageComp))
+		{
+			Result.ModifiedComponents.Add(FoliageComp);
+		}
+	}
+
+	return Result;
+}
+
+bool UE::LevelSnapshots::Foliage::Private::FFoliageRestorationInfo::ShouldSkipFoliageType(const FFoliageInfoData& SavedData) const
+{
+	return SavedData.GetImplType() != EFoliageImplType::StaticMesh;
+}
+
+bool UE::LevelSnapshots::Foliage::Private::FFoliageRestorationInfo::ShouldSerializeFoliageType(const FFoliageInfoData& SavedData) const
+{
+	if (bWasRecreated)
+	{
+		return true;
+	}
+	
+	if (TOptional<FName> CompName = SavedData.GetComponentName(); CompName && SavedData.GetImplType() == EFoliageImplType::StaticMesh)
+	{
+		const bool bIsAddedComponent = SnapshotComponentsToAdd.ContainsByPredicate(
+			[CompName](const TWeakObjectPtr<UActorComponent>& Component)
+			{
+				return Component.IsValid() && Component->GetFName().IsEqual(*CompName);
+			});
+		const bool bIsModifiedComponent = ModifiedComponents.ContainsByPredicate(
+			[CompName](const UActorComponent* Component)
+			{
+				return Component->GetFName().IsEqual(*CompName);
+			});
+
+		return bIsAddedComponent || bIsModifiedComponent;
+	}
+
+	// Actor foliage is currently not supported: just restore everything
+	UE_CLOG(SavedData.GetImplType() != EFoliageImplType::StaticMesh, LogLevelSnapshots, Warning, TEXT("Only static mesh foliage is supported"));
+	return false;
+}
diff --git a/Engine/Plugins/VirtualProduction/LevelSnapshots/Source/FoliageSupport/Private/FoliageSupport/Apply/FoliageRestorationInfo.h b/Engine/Plugins/VirtualProduction/LevelSnapshots/Source/FoliageSupport/Private/FoliageSupport/Apply/FoliageRestorationInfo.h
new file mode 100644
--- /dev/null
+++ b/Engine/Plugins/VirtualProduction/LevelSnapshots/Source/FoliageSupport/Private/FoliageSupport/Apply/FoliageRestorationInfo.h
@@ -0,0 +1,30 @@
+﻿// Copyright Epic Games, Inc. All Rights Reserved.
+
+#pragma once
+
+#include "CoreMinimal.h"
+
+class AInstancedFoliageActor;
+class UActorComponent;
+struct FPropertySelectionMap;
+
+namespace UE::LevelSnapshots::Foliage::Private
+{
+	class FFoliageInfoData;
+	
+	/** Tells us whether foliage types should be added / removed / serialized into AInstancedFoliageActor::FoliageInfos */
+	class FFoliageRestorationInfo
+	{
+		TArray<UActorComponent*> ModifiedComponents;
+		TArray<TWeakObjectPtr<UActorComponent>> EditorWorldComponentsToRemove;
+		TArray<TWeakObjectPtr<UActorComponent>> SnapshotComponentsToAdd;
+
+		bool bWasRecreated = false;
+	public:
+	
+		static FFoliageRestorationInfo From(AInstancedFoliageActor* Object, const FPropertySelectionMap& SelectionMap, bool bWasRecreated);
+
+		bool ShouldSkipFoliageType(const FFoliageInfoData& SavedData) const;
+		bool ShouldSerializeFoliageType(const FFoliageInfoData& SavedData) const;
+	};
+}
\ No newline at end of file
diff --git a/Engine/Plugins/VirtualProduction/LevelSnapshots/Source/FoliageSupport/Private/FoliageSupport/Data/FoliageInfoData.cpp b/Engine/Plugins/VirtualProduction/LevelSnapshots/Source/FoliageSupport/Private/FoliageSupport/Data/FoliageInfoData.cpp
new file mode 100644
--- /dev/null
+++ b/Engine/Plugins/VirtualProduction/LevelSnapshots/Source/FoliageSupport/Private/FoliageSupport/Data/FoliageInfoData.cpp
@@ -0,0 +1,85 @@
+﻿// Copyright Epic Games, Inc. All Rights Reserved.
+
+#include "FoliageSupport/Data/FoliageInfoData.h"
+
+#include "LevelSnapshotsLog.h"
+#include "SnapshotCustomVersion.h"
+
+#include "InstancedFoliageActor.h"
+#include "Serialization/Archive.h"
+#include "Serialization/MemoryReader.h"
+#include "Serialization/ObjectAndNameAsStringProxyArchive.h"
+
+namespace UE::LevelSnapshots::Foliage::Private
+{
+	static FString AsName(EFoliageImplType ImplType)
+	{
+		switch(ImplType)
+		{
+		case EFoliageImplType::StaticMesh:
+			return FString("StaticMesh");
+		case EFoliageImplType::Actor:
+			return FString("Actor");
+		case EFoliageImplType::ISMActor:
+			return FString("ISMActor");
+			
+		case EFoliageImplType::Unknown: 
+			return FString("Unknown");
+		default:
+			ensureMsgf(false, TEXT("Update this switch case."));
+			return FString("NewType (Update FoliageInfoData.cpp)");
+		}
+	}
+}
+
+FArchive& UE::LevelSnapshots::Foliage::Private::FFoliageInfoData::SerializeInternal(FArchive& Ar)
+{
+	Ar << ImplType;
+	Ar << ComponentName;
+	if (Ar.CustomVer(FSnapshotCustomVersion::GUID) < FSnapshotCustomVersion::CustomSubobjectSoftObjectPathRefactor)
+	{
+		Ar << SerializedData_DEPRECATED;
+	}
+	else
+	{
+		check(Ar.CustomVer(FSnapshotCustomVersion::GUID) >= FSnapshotCustomVersion::FoliageTypesUnreadable);
+		Ar << FoliageInfoArchiveSize;
+		Ar << ArchiveTellBeforeSerialization;
+	}
+	return Ar;
+}
+
+void UE::LevelSnapshots::Foliage::Private::FFoliageInfoData::FillDataMembersAndWriteFoliageInfoIntoArchive(FArchive& Archive, FFoliageInfo& DataToReadFrom)
+{
+	ArchiveTellBeforeSerialization = Archive.Tell();
+	ImplType = DataToReadFrom.Type;
+	
+	const bool bIsComponentValid = DataToReadFrom.GetComponent() != nullptr; 
+	ComponentName = bIsComponentValid ? DataToReadFrom.GetComponent()->GetFName() : FName(NAME_None);
+	UE_CLOG(!bIsComponentValid, LogLevelSnapshots, Warning, TEXT("Could not save component for ImplType %s"), *AsName(ImplType));
+
+	int64 ArchivePos = Archive.Tell();
+	Archive << DataToReadFrom;
+	FoliageInfoArchiveSize = Archive.Tell() - ArchivePos;
+}
+
+void UE::LevelSnapshots::Foliage::Private::FFoliageInfoData::ApplyTo(FArchive& Archive, FFoliageInfo& DataToWriteInto) const
+{
+	// Avoid foliage internal checks
+	DataToWriteInto.Implementation.Reset();
+	
+	if (Archive.CustomVer(FSnapshotCustomVersion::GUID) >= FSnapshotCustomVersion::CustomSubobjectSoftObjectPathRefactor)
+	{
+		Archive.Seek(ArchiveTellBeforeSerialization);
+		Archive << DataToWriteInto;
+	}
+	else
+	{
+		FMemoryReader MemoryReader(SerializedData_DEPRECATED, true);
+		FObjectAndNameAsStringProxyArchive RootArchive(MemoryReader, false);
+		RootArchive.SetCustomVersions(Archive.GetCustomVersions());
+		RootArchive << DataToWriteInto;
+	}
+	
+	DataToWriteInto.RecomputeHash();
+}
\ No newline at end of file
diff --git a/Engine/Plugins/VirtualProduction/LevelSnapshots/Source/FoliageSupport/Private/FoliageSupport/Data/FoliageInfoData.h b/Engine/Plugins/VirtualProduction/LevelSnapshots/Source/FoliageSupport/Private/FoliageSupport/Data/FoliageInfoData.h
new file mode 100644
--- /dev/null
+++ b/Engine/Plugins/VirtualProduction/LevelSnapshots/Source/FoliageSupport/Private/FoliageSupport/Data/FoliageInfoData.h
@@ -0,0 +1,50 @@
+﻿// Copyright Epic Games, Inc. All Rights Reserved.
+
+#pragma once
+
+#include "CoreMinimal.h"
+#include "InstancedFoliage.h"
+
+namespace UE::LevelSnapshots::Foliage::Private
+{
+	/** Data we'll send to FFoliageImpl::Serialize */
+	class FFoliageInfoData
+	{
+		EFoliageImplType ImplType = EFoliageImplType::Unknown;
+		FName ComponentName;
+		
+		/** Used before FSnapshotCustomVersion::CustomSubobjectSoftObjectPathRefactor */
+		TArray<uint8> SerializedData_DEPRECATED;
+		/** Added after FSnapshotCustomVersion::CustomSubobjectSoftObjectPathRefactor */
+		int64 FoliageInfoArchiveSize = 0;
+		
+		/**
+		 * What Ar.Tell returned at the beginning of CaptureDataAndWriteFoliageInfoIntoArchive (since 5.2, FoliageTypesUnreadable).
+		 * Used to jump across data to selectively apply data to foliage (since the user can choose to restore partially).
+		 */
+		int64 ArchiveTellBeforeSerialization;
+
+		FArchive& SerializeInternal(FArchive& Ar);
+	
+	public:
+
+		/** Saves FFoliageInfo into the archive and updates our data members. The our data members are NOT saved into the archive at this point. */
+		void FillDataMembersAndWriteFoliageInfoIntoArchive(FArchive& Archive, FFoliageInfo& DataToReadFrom);
+		/**
+		 * Given an archive containing data written by FillDataMembersAndWriteFoliageInfoIntoArchive previously and applies that data to the FFoliageInfo.
+		 * 5.0 and before still uses SerializedData_DEPRECATED for this while 5.2 and beyond seeks to ArchiveTellBeforeSerialization and then reads.
+		 */
+		void ApplyTo(FArchive& Archive, FFoliageInfo& DataToWriteInto) const;
+
+		EFoliageImplType GetImplType() const { return ImplType; }
+		TOptional<FName> GetComponentName() const { return ImplType == EFoliageImplType::StaticMesh && !ComponentName.IsNone() ? TOptional<FName>(ComponentName) : TOptional<FName>(); }
+		
+		int64 GetFoliageInfoArchiveSize() const { return FoliageInfoArchiveSize; }
+		int64 GetArchiveTellBeforeSerialization() const { return  ArchiveTellBeforeSerialization; }
+	
+		friend FArchive& operator<<(FArchive& Ar, FFoliageInfoData& MeshInfo)
+		{
+			return MeshInfo.SerializeInternal(Ar);
+		}
+	};
+}
diff --git a/Engine/Plugins/VirtualProduction/LevelSnapshots/Source/FoliageSupport/Private/FoliageSupport/Data/InstancedFoliageActorData.cpp b/Engine/Plugins/VirtualProduction/LevelSnapshots/Source/FoliageSupport/Private/FoliageSupport/Data/InstancedFoliageActorData.cpp
new file mode 100644
--- /dev/null
+++ b/Engine/Plugins/VirtualProduction/LevelSnapshots/Source/FoliageSupport/Private/FoliageSupport/Data/InstancedFoliageActorData.cpp
@@ -0,0 +1,105 @@
+﻿// Copyright Epic Games, Inc. All Rights Reserved.
+
+#include "FoliageSupport/Data/InstancedFoliageActorData.h"
+
+#include "Selection/PropertySelectionMap.h"
+#include "SnapshotCustomVersion.h"
+
+#include "FoliageType.h"
+#include "InstancedFoliageActor.h"
+#include "Serialization/Archive.h"
+
+namespace UE::LevelSnapshots::Foliage::Private
+{
+	static void SaveAsset(FArchive& Archive, UFoliageType* FoliageType, FFoliageInfo& FoliageInfo, TMap<TSoftObjectPtr<UFoliageType>, FFoliageInfoData>& FoliageAssets)
+	{
+		FFoliageInfoData& InfoData = FoliageAssets.Add(FoliageType);
+		InfoData.FillDataMembersAndWriteFoliageInfoIntoArchive(Archive, FoliageInfo);
+	}
+	
+	static void SaveSubobject(FArchive& Archive, UFoliageType* FoliageType, FFoliageInfo& FoliageInfo,TArray<FSubobjectFoliageInfoData>& SubobjectData)
+	{
+		FSubobjectFoliageInfoData Data;
+		Data.FillDataMembersAndSerializeFoliageTypeAndInfoIntoArchive(Archive, FoliageType, FoliageInfo);
+		SubobjectData.Emplace(Data);
+	}
+
+	void FInstancedFoliageActorData::CaptureData(FArchive& Archive, AInstancedFoliageActor* FoliageActor)
+	{
+		FInstancedFoliageActorData Data;
+		
+		// Reserve an int64 so later we can write how big the saved data was... See LoadAndApplyTo for details
+		const int64 LocBeforeWrite = Archive.Tell();
+		int64 Offset = 0;
+		Archive << Offset;
+
+		// Block 1 (see LoadAndApplyTo)
+		FoliageActor->ForEachFoliageInfo([&Archive, FoliageActor, &Data](UFoliageType* FoliageType, FFoliageInfo& FoliageInfo)
+		{
+			FFoliageImpl* Impl = FoliageInfo.Implementation.Get();
+			if (!FoliageType || !Impl)
+			{
+				return true;
+			}
+			
+			const bool bIsSubobject = FoliageType->IsIn(FoliageActor);
+			if (bIsSubobject)
+			{
+				SaveSubobject(Archive, FoliageType, FoliageInfo, Data.SubobjectData);
+			}
+			else
+			{
+				SaveAsset(Archive, FoliageType, FoliageInfo, Data.FoliageAssets);
+			}
+			
+			return true;
+		});
+
+		// ... Now go to the beginning and write how big the saved data was ...
+		Offset = Archive.Tell() - LocBeforeWrite;
+		Archive.Seek(LocBeforeWrite);
+		Archive << Offset;
+		// ... And make sure to go back to the end so anyone after us won't overwrite our data
+		Archive.Seek(LocBeforeWrite + Offset);
+
+		// Block 2 (see LoadAndApplyTo)
+		Archive << Data;
+	}
+
+	void FInstancedFoliageActorData::LoadAndApplyTo(FArchive& Archive, AInstancedFoliageActor* FoliageActor, const FPropertySelectionMap& SelectedProperties, bool bWasRecreated)
+	{
+		FInstancedFoliageActorData CurrentFoliageData;
+		if (Archive.CustomVer(FSnapshotCustomVersion::GUID) < FSnapshotCustomVersion::CustomSubobjectSoftObjectPathRefactor)
+		{
+			// In this version, all data was saved inline in FFoliageInfoData::SerializedData_DEPRECATED, so Archive << Data is enough
+			FInstancedFoliageActorData Data;
+			Archive << Data;
+			Data.ApplyData_Pre5dot2(Archive, FoliageActor, SelectedProperties, bWasRecreated);
+		}
+		else
+		{
+			// In this version, we have two blocks...
+			checkf(Archive.CustomVer(FSnapshotCustomVersion::GUID) >= FSnapshotCustomVersion::FoliageTypesUnreadable,
+				TEXT("Data from CustomSubobjectSoftObjectPathRefactor <= x < FoliageTypesUnreadable should have been discarded by the IActorSnapshotFilter implementation!")
+				);
+			
+			// ... block 1's size is saved at beginning. Contains "random" sub-blocks of of 1.1 FFoliageInfo or 1.2 FFoliageInfo followed by UFoliageType data
+			int64 Offset = 0;
+			Archive << Offset;
+			
+			// ... block 2: Archive << Data. FFoliageInfoData::ArchiveTellBeforeSerialization contains the sub-block locations to be used for FArchive::Seek() in ApplyTo_Post5dot2.
+			Archive.Seek(Offset);
+			FInstancedFoliageActorData Data;
+			Archive << Data;
+
+			Data.ApplyTo_5dot2(Archive, FoliageActor, SelectedProperties, bWasRecreated);
+		}
+	}
+
+	FArchive& operator<<(FArchive& Archive, FInstancedFoliageActorData& Data)
+	{
+		Archive << Data.FoliageAssets;
+		Archive << Data.SubobjectData;
+		return Archive;
+	}
+}
diff --git a/Engine/Plugins/VirtualProduction/LevelSnapshots/Source/FoliageSupport/Private/FoliageSupport/Data/InstancedFoliageActorData.h b/Engine/Plugins/VirtualProduction/LevelSnapshots/Source/FoliageSupport/Private/FoliageSupport/Data/InstancedFoliageActorData.h
new file mode 100644
--- /dev/null
+++ b/Engine/Plugins/VirtualProduction/LevelSnapshots/Source/FoliageSupport/Private/FoliageSupport/Data/InstancedFoliageActorData.h
@@ -0,0 +1,48 @@
+﻿// Copyright Epic Games, Inc. All Rights Reserved.
+
+#pragma once
+
+#include "FoliageSupport/Data/SubobjectFoliageInfoData.h"
+#include "FoliageSupport/Data/FoliageInfoData.h"
+#include "UObject/SoftObjectPtr.h"
+
+class FArchive;
+class UFoliageType;
+struct FPropertySelectionMap;
+
+namespace UE::LevelSnapshots::Foliage::Private
+{
+	/** Handles saving of AInstancedFoliageActor::FoliageInfos */
+	class FInstancedFoliageActorData
+	{
+		/**
+		 * Foliage type asset references.
+		 * 
+		 * Consists of foliage type assets added via the green Add button or drag-drop a foliage type or if you drag-drop
+		 * a foliage type from the content browser.
+		 */
+		TMap<TSoftObjectPtr<UFoliageType>, FFoliageInfoData> FoliageAssets;
+
+		/**
+		 * Saves foliage type subobjects.
+		 * Consists of foliage types created when you drag-drop a static mesh or add Blueprint derived foliage types.
+		 */
+		TArray<FSubobjectFoliageInfoData> SubobjectData;
+		
+		/** Serialization version pre 5.1. Data saved in 5.1 is discarded. Looks up saved foliage types and restores the passed in foliage actor. */
+		void ApplyData_Pre5dot2(FArchive& Archive, AInstancedFoliageActor* FoliageActor, const FPropertySelectionMap& SelectedProperties, bool bWasRecreated) const;
+
+		/** Serialization version starting 5.2. Data saved in 5.1 is discarded. */
+		void ApplyTo_5dot2(FArchive& Archive, AInstancedFoliageActor* FoliageActor, const FPropertySelectionMap& SelectedProperties, bool bWasRecreated) const;
+
+		friend FArchive& operator<<(FArchive& Archive, FInstancedFoliageActorData& Data);
+		
+	public:
+
+		/** Saves foliage actor into this object */
+		static void CaptureData(FArchive& Archive, AInstancedFoliageActor* FoliageActor);
+
+		/** Loads the saved foliage data and applies it to FoliageActor respecting the selection in SelectedProperties. */
+		static void LoadAndApplyTo(FArchive& Archive, AInstancedFoliageActor* FoliageActor, const FPropertySelectionMap& SelectedProperties, bool bWasRecreated);
+	};
+}
diff --git a/Engine/Plugins/VirtualProduction/LevelSnapshots/Source/FoliageSupport/Private/FoliageSupport/Data/SubobjectFoliageInfoData.cpp b/Engine/Plugins/VirtualProduction/LevelSnapshots/Source/FoliageSupport/Private/FoliageSupport/Data/SubobjectFoliageInfoData.cpp
new file mode 100644
--- /dev/null
+++ b/Engine/Plugins/VirtualProduction/LevelSnapshots/Source/FoliageSupport/Private/FoliageSupport/Data/SubobjectFoliageInfoData.cpp
@@ -0,0 +1,83 @@
+﻿// Copyright Epic Games, Inc. All Rights Reserved.
+
+#include "FoliageSupport/Data/SubobjectFoliageInfoData.h"
+
+#include "FoliageInfoData.h"
+#include "LevelSnapshotsLog.h"
+#include "SnapshotCustomVersion.h"
+
+#include "FoliageType_InstancedStaticMesh.h"
+#include "InstancedFoliageActor.h"
+#include "Serialization/Archive.h"
+#include "Serialization/MemoryReader.h"
+#include "Serialization/ObjectAndNameAsStringProxyArchive.h"
+
+FArchive& UE::LevelSnapshots::Foliage::Private::FSubobjectFoliageInfoData::SerializeInternal(FArchive& Ar)
+{
+	Ar << static_cast<FFoliageInfoData&>(*this);
+	Ar << Class;
+	Ar << SubobjectName;
+	if (Ar.CustomVer(FSnapshotCustomVersion::GUID) < FSnapshotCustomVersion::CustomSubobjectSoftObjectPathRefactor)
+	{
+		Ar << SerializedSubobjectData_DEPRECATED;
+	}
+	else
+	{
+		Ar << FoliageTypeArchiveSize;
+	}
+	return Ar;
+}
+
+void UE::LevelSnapshots::Foliage::Private::FSubobjectFoliageInfoData::FillDataMembersAndSerializeFoliageTypeAndInfoIntoArchive(FArchive& Archive, UFoliageType* FoliageSubobject, FFoliageInfo& FoliageInfo)
+{
+	FFoliageInfoData::FillDataMembersAndWriteFoliageInfoIntoArchive(Archive, FoliageInfo);
+
+	Class = FoliageSubobject->GetClass();
+	SubobjectName = FoliageSubobject->GetFName();
+
+	const int64 ArchivePos = Archive.Tell();
+	FoliageSubobject->Serialize(Archive);
+	FoliageTypeArchiveSize += Archive.Tell() - ArchivePos;
+}
+
+UFoliageType* UE::LevelSnapshots::Foliage::Private::FSubobjectFoliageInfoData::FindOrRecreateSubobject(AInstancedFoliageActor* Outer) const
+{
+	if (UObject* FoundObject = StaticFindObjectFast(nullptr, Outer, SubobjectName))
+	{
+		UE_CLOG(FoundObject->GetClass() != Class, LogLevelSnapshots, Warning, TEXT("Name collision for foliage type %s"), *FoundObject->GetPathName());
+		return Cast<UFoliageType>(FoundObject);
+	}
+
+	return NewObject<UFoliageType>(Outer, Class, SubobjectName, RF_Transactional);
+}
+
+bool UE::LevelSnapshots::Foliage::Private::FSubobjectFoliageInfoData::ApplyToFoliageType(FArchive& Archive, UFoliageType* FoliageSubobject) const
+{
+	if (Archive.CustomVer(FSnapshotCustomVersion::GUID) >= FSnapshotCustomVersion::CustomSubobjectSoftObjectPathRefactor)
+	{
+		Archive.Seek(GetArchiveTellBeforeSerialization() + GetFoliageInfoArchiveSize());
+		FoliageSubobject->Serialize(Archive);
+	}
+	else
+	{
+		FMemoryReader MemoryReader(SerializedSubobjectData_DEPRECATED, true);
+		constexpr bool bLoadIfFindFails = true;
+		FObjectAndNameAsStringProxyArchive RootArchive(MemoryReader, bLoadIfFindFails);
+		FoliageSubobject->Serialize(RootArchive);
+	}
+
+	if (UFoliageType_InstancedStaticMesh* MeshFoliageType = Cast<UFoliageType_InstancedStaticMesh>(FoliageSubobject)
+		; ensureMsgf(MeshFoliageType, TEXT("Only static mesh foliage is supported right now")))
+	{
+		// There used to be a bug where the asset registry was not aware
+		UE_CLOG(!MeshFoliageType->Mesh, LogLevelSnapshots, Error, TEXT("Reference to foliage mesh was corrupted."));
+		return MeshFoliageType->Mesh != nullptr;
+	}
+
+	return true;
+}
+
+void UE::LevelSnapshots::Foliage::Private::FSubobjectFoliageInfoData::ApplyToFoliageInfo(FArchive& Archive, FFoliageInfo& DataToWriteInto) const
+{
+	FFoliageInfoData::ApplyTo(Archive, DataToWriteInto);
+}
diff --git a/Engine/Plugins/VirtualProduction/LevelSnapshots/Source/FoliageSupport/Private/FoliageSupport/Data/SubobjectFoliageInfoData.h b/Engine/Plugins/VirtualProduction/LevelSnapshots/Source/FoliageSupport/Private/FoliageSupport/Data/SubobjectFoliageInfoData.h
new file mode 100644
--- /dev/null
+++ b/Engine/Plugins/VirtualProduction/LevelSnapshots/Source/FoliageSupport/Private/FoliageSupport/Data/SubobjectFoliageInfoData.h
@@ -0,0 +1,45 @@
+﻿// Copyright Epic Games, Inc. All Rights Reserved.
+
+#pragma once
+
+#include "FoliageSupport/Data/FoliageInfoData.h"
+#include "Containers/Array.h"
+#include "Templates/SubclassOf.h"
+#include "UObject/NameTypes.h"
+
+class FArchive;
+class UFoliageType;
+
+namespace UE::LevelSnapshots::Foliage::Private
+{
+	class FSubobjectFoliageInfoData;
+	FArchive& operator<<(FArchive& Ar, FSubobjectFoliageInfoData& Data);
+	
+	class FSubobjectFoliageInfoData : public FFoliageInfoData
+	{
+		TSubclassOf<UFoliageType> Class;
+		FName SubobjectName;
+		
+		/** Used before FSnapshotCustomVersion::CustomSubobjectSoftObjectPathRefactor*/
+		TArray<uint8> SerializedSubobjectData_DEPRECATED;
+		/** Added after FSnapshotCustomVersion::CustomSubobjectSoftObjectPathRefactor */
+		int64 FoliageTypeArchiveSize = 0;
+
+		FArchive& SerializeInternal(FArchive& Ar);
+	
+	public:
+
+		void FillDataMembersAndSerializeFoliageTypeAndInfoIntoArchive(FArchive& Archive, UFoliageType* FoliageSubobject, FFoliageInfo& FoliageInfo);
+	
+		UFoliageType* FindOrRecreateSubobject(AInstancedFoliageActor* Outer) const;
+		bool ApplyToFoliageType(FArchive& Archive, UFoliageType* FoliageSubobject) const;
+		void ApplyToFoliageInfo(FArchive& Archive, FFoliageInfo& DataToWriteInto) const;
+
+		int64 GetFoliageTypeArchiveSize() const { return FoliageTypeArchiveSize; }
+		
+		friend FArchive& operator<<(FArchive& Ar, FSubobjectFoliageInfoData& Data)
+		{
+			return Data.SerializeInternal(Ar);
+		}
+	};
+}
diff --git a/Engine/Plugins/VirtualProduction/LevelSnapshots/Source/FoliageSupport/Private/FoliageSupport/FoliageInfoData.cpp b/Engine/Plugins/VirtualProduction/LevelSnapshots/Source/FoliageSupport/Private/FoliageSupport/FoliageInfoData.cpp
deleted file mode 100644
diff --git a/Engine/Plugins/VirtualProduction/LevelSnapshots/Source/FoliageSupport/Private/FoliageSupport/FoliageInfoData.h b/Engine/Plugins/VirtualProduction/LevelSnapshots/Source/FoliageSupport/Private/FoliageSupport/FoliageInfoData.h
deleted file mode 100644
diff --git a/Engine/Plugins/VirtualProduction/LevelSnapshots/Source/FoliageSupport/Private/FoliageSupport/FoliageRestorationInfo.cpp b/Engine/Plugins/VirtualProduction/LevelSnapshots/Source/FoliageSupport/Private/FoliageSupport/FoliageRestorationInfo.cpp
deleted file mode 100644
diff --git a/Engine/Plugins/VirtualProduction/LevelSnapshots/Source/FoliageSupport/Private/FoliageSupport/FoliageRestorationInfo.h b/Engine/Plugins/VirtualProduction/LevelSnapshots/Source/FoliageSupport/Private/FoliageSupport/FoliageRestorationInfo.h
deleted file mode 100644
diff --git a/Engine/Plugins/VirtualProduction/LevelSnapshots/Source/FoliageSupport/Private/FoliageSupport/FoliageSupport.cpp b/Engine/Plugins/VirtualProduction/LevelSnapshots/Source/FoliageSupport/Private/FoliageSupport/FoliageSupport.cpp
--- a/Engine/Plugins/VirtualProduction/LevelSnapshots/Source/FoliageSupport/Private/FoliageSupport/FoliageSupport.cpp
+++ b/Engine/Plugins/VirtualProduction/LevelSnapshots/Source/FoliageSupport/Private/FoliageSupport/FoliageSupport.cpp
@@ -2,40 +2,45 @@
 
 #include "FoliageSupport/FoliageSupport.h"
 
-#include "Engine/Level.h"
-#include "FoliageSupport/InstancedFoliageActorData.h"
+#include "FoliageSupport/Data/InstancedFoliageActorData.h"
+#include "Params/ObjectSnapshotSerializationData.h"
 #include "Selection/PropertySelection.h"
 #include "Selection/PropertySelectionMap.h"
-#include "Params/ObjectSnapshotSerializationData.h"
+#include "SnapshotCustomVersion.h"
+#include "WorldSnapshotData.h"
 
+#include "ActorPartition/ActorPartitionSubsystem.h"
 #include "Components/InstancedStaticMeshComponent.h"
 #include "Components/HierarchicalInstancedStaticMeshComponent.h"
+#include "Engine/Level.h"
 #include "FoliageHelper.h"
+#include "FoliageLevelSnapshotsConsoleVariables.h"
 #include "ILevelSnapshotsModule.h"
 #include "InstancedFoliageActor.h"
-#include "SnapshotCustomVersion.h"
-#include "ActorPartition/ActorPartitionSubsystem.h"
 #if WITH_EDITOR
 #include "FoliageEditModule.h"
 #endif
 
-namespace UE::LevelSnapshots::Foliage::Private::Internal
+#define LOCTEXT_NAMESPACE "LevelSnapshots.FoliageSupport"
+
+namespace UE::LevelSnapshots::Foliage::Private
 {
 	static void EnableRequiredFoliageProperties(ILevelSnapshotsModule& Module)
 	{
 		FProperty* InstanceReorderTable = UInstancedStaticMeshComponent::StaticClass()->FindPropertyByName(GET_MEMBER_NAME_CHECKED(UInstancedStaticMeshComponent, InstanceReorderTable));
-		
-		FProperty* SortedInstances = UHierarchicalInstancedStaticMeshComponent::StaticClass()->FindPropertyByName(GET_MEMBER_NAME_CHECKED(UHierarchicalInstancedStaticMeshComponent, SortedInstances));
-		FProperty* NumBuiltInstances = UHierarchicalInstancedStaticMeshComponent::StaticClass()->FindPropertyByName(GET_MEMBER_NAME_CHECKED(UHierarchicalInstancedStaticMeshComponent, NumBuiltInstances));
-		FProperty* NumBuiltRenderInstances = UHierarchicalInstancedStaticMeshComponent::StaticClass()->FindPropertyByName(GET_MEMBER_NAME_CHECKED(UHierarchicalInstancedStaticMeshComponent, NumBuiltRenderInstances));
-		FProperty* BuiltInstanceBounds = UHierarchicalInstancedStaticMeshComponent::StaticClass()->FindPropertyByName(GET_MEMBER_NAME_CHECKED(UHierarchicalInstancedStaticMeshComponent, BuiltInstanceBounds));
-		FProperty* UnbuiltInstanceBounds = UHierarchicalInstancedStaticMeshComponent::StaticClass()->FindPropertyByName(GET_MEMBER_NAME_CHECKED(UHierarchicalInstancedStaticMeshComponent, UnbuiltInstanceBounds));
-		FProperty* UnbuiltInstanceBoundsList = UHierarchicalInstancedStaticMeshComponent::StaticClass()->FindPropertyByName(GET_MEMBER_NAME_CHECKED(UHierarchicalInstancedStaticMeshComponent, UnbuiltInstanceBoundsList));
-		FProperty* bEnableDensityScaling = UHierarchicalInstancedStaticMeshComponent::StaticClass()->FindPropertyByName(GET_MEMBER_NAME_CHECKED(UHierarchicalInstancedStaticMeshComponent, bEnableDensityScaling));
-		FProperty* OcclusionLayerNumNodes = UHierarchicalInstancedStaticMeshComponent::StaticClass()->FindPropertyByName(GET_MEMBER_NAME_CHECKED(UHierarchicalInstancedStaticMeshComponent, OcclusionLayerNumNodes));
-		FProperty* CacheMeshExtendedBounds = UHierarchicalInstancedStaticMeshComponent::StaticClass()->FindPropertyByName(GET_MEMBER_NAME_CHECKED(UHierarchicalInstancedStaticMeshComponent, CacheMeshExtendedBounds));
-		FProperty* bDisableCollision = UHierarchicalInstancedStaticMeshComponent::StaticClass()->FindPropertyByName(GET_MEMBER_NAME_CHECKED(UHierarchicalInstancedStaticMeshComponent, bDisableCollision));
-		FProperty* InstanceCountToRender = UHierarchicalInstancedStaticMeshComponent::StaticClass()->FindPropertyByName(GET_MEMBER_NAME_CHECKED(UHierarchicalInstancedStaticMeshComponent, InstanceCountToRender));
+
+		UClass* HierarchicalClass				= UHierarchicalInstancedStaticMeshComponent::StaticClass(); 
+		FProperty* SortedInstances				= HierarchicalClass->FindPropertyByName(GET_MEMBER_NAME_CHECKED(UHierarchicalInstancedStaticMeshComponent, SortedInstances));
+		FProperty* NumBuiltInstances			= HierarchicalClass->FindPropertyByName(GET_MEMBER_NAME_CHECKED(UHierarchicalInstancedStaticMeshComponent, NumBuiltInstances));
+		FProperty* NumBuiltRenderInstances		= HierarchicalClass->FindPropertyByName(GET_MEMBER_NAME_CHECKED(UHierarchicalInstancedStaticMeshComponent, NumBuiltRenderInstances));
+		FProperty* BuiltInstanceBounds			= HierarchicalClass->FindPropertyByName(GET_MEMBER_NAME_CHECKED(UHierarchicalInstancedStaticMeshComponent, BuiltInstanceBounds));
+		FProperty* UnbuiltInstanceBounds		= HierarchicalClass->FindPropertyByName(GET_MEMBER_NAME_CHECKED(UHierarchicalInstancedStaticMeshComponent, UnbuiltInstanceBounds));
+		FProperty* UnbuiltInstanceBoundsList	= HierarchicalClass->FindPropertyByName(GET_MEMBER_NAME_CHECKED(UHierarchicalInstancedStaticMeshComponent, UnbuiltInstanceBoundsList));
+		FProperty* bEnableDensityScaling		= HierarchicalClass->FindPropertyByName(GET_MEMBER_NAME_CHECKED(UHierarchicalInstancedStaticMeshComponent, bEnableDensityScaling));
+		FProperty* OcclusionLayerNumNodes		= HierarchicalClass->FindPropertyByName(GET_MEMBER_NAME_CHECKED(UHierarchicalInstancedStaticMeshComponent, OcclusionLayerNumNodes));
+		FProperty* CacheMeshExtendedBounds		= HierarchicalClass->FindPropertyByName(GET_MEMBER_NAME_CHECKED(UHierarchicalInstancedStaticMeshComponent, CacheMeshExtendedBounds));
+		FProperty* bDisableCollision			= HierarchicalClass->FindPropertyByName(GET_MEMBER_NAME_CHECKED(UHierarchicalInstancedStaticMeshComponent, bDisableCollision));
+		FProperty* InstanceCountToRender		= HierarchicalClass->FindPropertyByName(GET_MEMBER_NAME_CHECKED(UHierarchicalInstancedStaticMeshComponent, InstanceCountToRender));
 		
 		if (ensure(InstanceReorderTable))
 		{
@@ -55,40 +60,35 @@ namespace UE::LevelSnapshots::Foliage::Private::Internal
 			});
 		}
 	}
-}
 
-void UE::LevelSnapshots::Foliage::Private::FFoliageSupport::Register(ILevelSnapshotsModule& Module)
-{
-	Internal::EnableRequiredFoliageProperties(Module);
-	
-	const TSharedRef<FFoliageSupport> FoliageSupport = MakeShared<FFoliageSupport>();
-	Module.RegisterRestorabilityOverrider(FoliageSupport);
-	Module.RegisterRestorationListener(FoliageSupport);
-	Module.RegisterCustomObjectSerializer(AInstancedFoliageActor::StaticClass(), FoliageSupport);
-}
+	static UFoliageType* FindFoliageInfoFor(UHierarchicalInstancedStaticMeshComponent* Component)
+	{
+		AInstancedFoliageActor* Foliage = Cast<AInstancedFoliageActor>(Component->GetOwner());
+		if (!LIKELY(Foliage))
+		{
+			return nullptr;
+		}
 
-UE::LevelSnapshots::ISnapshotRestorabilityOverrider::ERestorabilityOverride UE::LevelSnapshots::Foliage::Private::FFoliageSupport::IsActorDesirableForCapture(const AActor* Actor)
-{
-	// Foliage's not allowed by default because it is hidden from the scene outliner
-	return Actor->GetClass() == AInstancedFoliageActor::StaticClass() || FFoliageHelper::IsOwnedByFoliage(Actor)
-		? ERestorabilityOverride::Allow : ERestorabilityOverride::DoNotCare;
-}
+		for (auto FoliageIt = Foliage->GetFoliageInfos().CreateConstIterator(); FoliageIt; ++FoliageIt)
+		{
+			if (FoliageIt->Value->Implementation->IsOwnedComponent(Component))
+			{
+				return FoliageIt->Key;
+			}
+		}
 
-void UE::LevelSnapshots::Foliage::Private::FFoliageSupport::OnTakeSnapshot(UObject* EditorObject, ICustomSnapshotSerializationData& DataStorage)
-{
-	AInstancedFoliageActor* FoliageActor = Cast<AInstancedFoliageActor>(EditorObject);
-	check(FoliageActor);
+		return nullptr;
+	}
 
-	DataStorage.WriteObjectAnnotation(FObjectAnnotator::CreateLambda([FoliageActor](FArchive& Archive)
+	static UFoliageType* FindFoliageInfoFor(UActorComponent* Component)
 	{
-		FInstancedFoliageActorData FoliageData;
-		FoliageData.Save(Archive, FoliageActor);
-	}));
-}
-
+		if (UHierarchicalInstancedStaticMeshComponent* Comp = Cast<UHierarchicalInstancedStaticMeshComponent>(Component))
+		{
+			return FindFoliageInfoFor(Comp);
+		}
+		return nullptr;
+	}
 
-namespace UE::LevelSnapshots::FoliageSupport::Internal
-{
 	static void RebuildChangedFoliageComponents(AInstancedFoliageActor* FoliageActor, const FPropertySelectionMap& SelectionMap, bool bWasRecreated)
 	{
 		const FRestorableObjectSelection FoliageSelection = SelectionMap.GetObjectSelection(FoliageActor);
@@ -108,15 +108,16 @@ namespace UE::LevelSnapshots::FoliageSupport::Internal
 			
 			if (bWasRecreated || bHasChangedProperties || bComponentWasRecreated)
 			{
-				// Unregister and register recreates render state: otherwise instances won't show
-				Comp->UnregisterComponent();
-				
-				// This recomputes transforms and shadows
-				constexpr bool bAsync = false;
-				constexpr bool bForceUpdate = true;
-				Comp->BuildTreeIfOutdated(bAsync, bForceUpdate);
-				
-				Comp->RegisterComponent();
+				// Getting foliage to update its visual state is quite involved and this is the easiest way...
+				// Feel free to adjust if you come up with an easier one
+				TArray<FTransform> InstanceTransforms;
+				InstanceTransforms.SetNumUninitialized(Comp->GetInstanceCount());
+				for (int32 i = 0; i < Comp->GetInstanceCount(); ++i)
+				{
+					Comp->GetInstanceTransform(i, InstanceTransforms[i]);
+				}
+				Comp->ClearInstances();
+				Comp->AddInstances(InstanceTransforms, false);
 			}
 		}
 	}
@@ -128,139 +129,200 @@ namespace UE::LevelSnapshots::FoliageSupport::Internal
 		FoliageEditModule.UpdateMeshList();
 #endif
 	}
-}
 
-void UE::LevelSnapshots::Foliage::Private::FFoliageSupport::PostApplyToEditorObject(UObject* Object, const ICustomSnapshotSerializationData& DataStorage, const FPropertySelectionMap& SelectionMap)
-{
-	AInstancedFoliageActor* FoliageActor = Cast<AInstancedFoliageActor>(Object);
-	check(FoliageActor);
-	
-	// Track this actor for safety so we know for sure that the functions were called in the order we expected them to
-	CurrentFoliageActor = FoliageActor;
-	DataStorage.ReadObjectAnnotation(FObjectAnnotator::CreateLambda([this, &SelectionMap, FoliageActor](FArchive& Archive)
+	void FFoliageSupport::Register(ILevelSnapshotsModule& Module)
 	{
-		FInstancedFoliageActorData CurrentFoliageData;
-		if (Archive.CustomVer(FSnapshotCustomVersion::GUID) < FSnapshotCustomVersion::CustomSubobjectSoftObjectPathRefactor)
-		{
-			Archive << CurrentFoliageData;
-		}
-		
-#if WITH_EDITOR
-		// Actor might not have been modified because none of its uproperties were changed
-		FoliageActor->Modify();
-#endif
-
-		const bool bWasRecreated = SelectionMap.GetDeletedActorsToRespawn().Contains(FoliageActor);
-		CurrentFoliageData.ApplyTo(Archive, FoliageActor, SelectionMap, bWasRecreated);
-		FoliageSupport::Internal::RebuildChangedFoliageComponents(FoliageActor, SelectionMap, bWasRecreated);
+		EnableRequiredFoliageProperties(Module);
+	
+		const TSharedRef<FFoliageSupport> FoliageSupport = MakeShared<FFoliageSupport>();
+		Module.RegisterRestorabilityOverrider(FoliageSupport);
+		Module.RegisterRestorationListener(FoliageSupport);
+		Module.RegisterCustomObjectSerializer(AInstancedFoliageActor::StaticClass(), FoliageSupport);
+		Module.RegisterGlobalActorFilter(FoliageSupport);
+	}
 
-		FoliageSupport::Internal::UpdateFoliageUI();
-	}));
+	ISnapshotRestorabilityOverrider::ERestorabilityOverride FFoliageSupport::IsActorDesirableForCapture(const AActor* Actor)
+	{
+		// Foliage's not allowed by default because it is hidden from the scene outliner
+		return Actor->GetClass() == AInstancedFoliageActor::StaticClass() || FFoliageHelper::IsOwnedByFoliage(Actor)
+			? ERestorabilityOverride::Allow : ERestorabilityOverride::DoNotCare;
+	}
 
-	// Rest is done in PostApplySnapshotToActor (need access to the property selection map)
-}
+	void FFoliageSupport::OnTakeSnapshot(UObject* EditorObject, ICustomSnapshotSerializationData& DataStorage)
+	{
+		AInstancedFoliageActor* FoliageActor = Cast<AInstancedFoliageActor>(EditorObject);
+		check(FoliageActor);
 
-void UE::LevelSnapshots::Foliage::Private::FFoliageSupport::PostApplySnapshotToActor(const FApplySnapshotToActorParams& Params)
-{
-	// This is the order of operations up until now
-	// 1. OnTakeSnapshot > Set initial data
-	// 2. User requests apply to world
-	// 3. FindOrRecreateSubobjectInEditorWorld allocates the UFoliageType subobjects
-	// 4. PreRemoveComponent
-	// 5. PostApplySnapshotProperties loads the data that OnTakeSnapshot saved
+		DataStorage.WriteObjectAnnotation(FObjectAnnotator::CreateLambda([FoliageActor](FArchive& Archive)
+		{
+			FInstancedFoliageActorData FoliageData;
+			FoliageData.CaptureData(Archive, FoliageActor);
+		}));
+	}
 	
-	if (AInstancedFoliageActor* FoliageActor = Cast<AInstancedFoliageActor>(Params.Actor))
+	void FFoliageSupport::PostApplyToEditorObject(UObject* Object, const ICustomSnapshotSerializationData& DataStorage, const FPropertySelectionMap& SelectionMap)
 	{
-		checkf(FoliageActor == CurrentFoliageActor.Get(), TEXT("PostApplySnapshotToActor was not directly followed by PostApplySnapshotProperties. Investigate."));
-		CurrentFoliageActor.Reset();
+		AInstancedFoliageActor* FoliageActor = Cast<AInstancedFoliageActor>(Object);
+		check(FoliageActor);
+	
+		// Track this actor for safety so we know for sure that the functions were called in the order we expected them to
+		CurrentFoliageActor = FoliageActor;
+		DataStorage.ReadObjectAnnotation(FObjectAnnotator::CreateLambda([this, &SelectionMap, FoliageActor](FArchive& Archive)
+		{
+	#if WITH_EDITOR
+			// Actor might not have been modified because none of its uproperties were changed
+			FoliageActor->Modify();
+			// Prevents: 1. Paint 1 instance 2. Take snapshot 3. Paint many instances 4. Select all 5. Apply snapshot > Crash due to selection index out of bounds
+			FoliageActor->ForEachFoliageInfo([](UFoliageType*, FFoliageInfo& Info)
+			{
+				Info.ClearSelection();
+				return true;
+			});
+	#endif
+			const bool bWasRecreated = SelectionMap.GetDeletedActorsToRespawn().Contains(FoliageActor);
+			FInstancedFoliageActorData::LoadAndApplyTo(Archive, FoliageActor, SelectionMap, bWasRecreated);
+
+			ActorsNeedingRefresh.Add(FoliageActor);
+			RebuildChangedFoliageComponents(FoliageActor, SelectionMap, bWasRecreated);
+			UpdateFoliageUI();
+		}));
+
+		// Rest is done in PostApplySnapshotToActor (need access to the property selection map)
 	}
-}
 
-void UE::LevelSnapshots::Foliage::Private::FFoliageSupport::PreApplySnapshotProperties(const FApplySnapshotPropertiesParams& Params)
-{
-	// RF_Transactional is temporarily required because Level Snapshots changes components differently than Foliage originally designed.
-	// Needed for correct undo / redo 
-	if (UHierarchicalInstancedStaticMeshComponent* Comp = Cast<UHierarchicalInstancedStaticMeshComponent>(Params.Object); Comp && CurrentFoliageActor.IsValid() && Comp->IsIn(CurrentFoliageActor.Get()))
+	void FFoliageSupport::PostApplySnapshot(const FApplySnapshotParams& Params)
 	{
-		Comp->SetFlags(RF_Transactional);
+		// Without this step all foliage will teleport visually upon being selected by the user; its cache needs to be updated:
+		// 1. Add 1 instance 2. Take snapshot 3. Paint many instances 4. Restore 5. Lasso select all > The instance jumps to a random place until you paint
+		for (TWeakObjectPtr<AInstancedFoliageActor> FoliageActor : ActorsNeedingRefresh)
+		{
+			if (!FoliageActor.IsValid())
+			{
+				return;
+			}
+			
+			FoliageActor->ForEachFoliageInfo([](UFoliageType*, FFoliageInfo& Info)
+			{
+				constexpr bool bAsync = true; 
+				constexpr bool bForce = true; // Must always run
+				Info.Refresh(bAsync, bForce);
+				return true;
+			});
+		}
+		ActorsNeedingRefresh.Empty();
 	}
-}
 
-void UE::LevelSnapshots::Foliage::Private::FFoliageSupport::PostApplySnapshotProperties(const FApplySnapshotPropertiesParams& Params)
-{
-	if (UHierarchicalInstancedStaticMeshComponent* Comp = Cast<UHierarchicalInstancedStaticMeshComponent>(Params.Object); Comp && CurrentFoliageActor.IsValid() && Comp->IsIn(CurrentFoliageActor.Get()))
+	void FFoliageSupport::PostApplySnapshotToActor(const FApplySnapshotToActorParams& Params)
 	{
-		Comp->ClearFlags(RF_Transactional);
+		// This is the order of operations up until now
+		// 1. OnTakeSnapshot > Set initial data
+		// 2. User requests apply to world
+		// 3. FindOrRecreateSubobjectInEditorWorld allocates the UFoliageType subobjects
+		// 4. PreRemoveComponent
+		// 5. PostApplySnapshotProperties loads the data that OnTakeSnapshot saved
+	
+		if (AInstancedFoliageActor* FoliageActor = Cast<AInstancedFoliageActor>(Params.Actor))
+		{
+			checkf(FoliageActor == CurrentFoliageActor.Get(), TEXT("PostApplySnapshotToActor was not directly followed by PostApplySnapshotProperties. Investigate."));
+			CurrentFoliageActor.Reset();
+		}
 	}
-}
 
-void UE::LevelSnapshots::Foliage::Private::FFoliageSupport::PreRecreateActor(UWorld* World, TSubclassOf<AActor> ActorClass, FActorSpawnParameters& InOutSpawnParameters)
-{
-	if (ActorClass->IsChildOf(AInstancedFoliageActor::StaticClass()))
+	void FFoliageSupport::PreApplySnapshotProperties(const FApplySnapshotPropertiesParams& Params)
 	{
-		InOutSpawnParameters.bCreateActorPackage = true;
-
-		// See AInstancedFoliageActor::GetDefault
-		ULevel* CurrentLevel = World->GetCurrentLevel();
-		if (CurrentLevel && CurrentLevel->GetWorld()->GetSubsystem<UActorPartitionSubsystem>()->IsLevelPartition())
+		// RF_Transactional is temporarily required because Level Snapshots changes components differently than Foliage originally designed.
+		// Needed for correct undo / redo 
+		if (UHierarchicalInstancedStaticMeshComponent* Comp = Cast<UHierarchicalInstancedStaticMeshComponent>(Params.Object); Comp && CurrentFoliageActor.IsValid() && Comp->IsIn(CurrentFoliageActor.Get()))
 		{
-			InOutSpawnParameters.ObjectFlags |= RF_Transient;
+			Comp->SetFlags(RF_Transactional);
 		}
 	}
-}
 
-void UE::LevelSnapshots::Foliage::Private::FFoliageSupport::PostRecreateActor(AActor* RecreatedActor)
-{
-	if (AInstancedFoliageActor* FoliageActor = Cast<AInstancedFoliageActor>(RecreatedActor))
+	void FFoliageSupport::PostApplySnapshotProperties(const FApplySnapshotPropertiesParams& Params)
 	{
-		RecreatedActor->GetLevel()->InstancedFoliageActor = FoliageActor;
+		if (UHierarchicalInstancedStaticMeshComponent* Comp = Cast<UHierarchicalInstancedStaticMeshComponent>(Params.Object); Comp && CurrentFoliageActor.IsValid() && Comp->IsIn(CurrentFoliageActor.Get()))
+		{
+			Comp->ClearFlags(RF_Transactional);
+		}
 	}
-}
 
-void UE::LevelSnapshots::Foliage::Private::FFoliageSupport::PreRemoveActor(const FPreRemoveActorParams& Params)
-{
-	if (AInstancedFoliageActor* FoliageActor = Cast<AInstancedFoliageActor>(Params.ActorToRemove))
+	void FFoliageSupport::PreRecreateActor(UWorld* World, TSubclassOf<AActor> ActorClass, FActorSpawnParameters& InOutSpawnParameters)
 	{
-		FoliageActor->GetLevel()->InstancedFoliageActor.Reset();
+		if (ActorClass->IsChildOf(AInstancedFoliageActor::StaticClass()))
+		{
+			InOutSpawnParameters.bCreateActorPackage = true;
+
+			// See AInstancedFoliageActor::GetDefault
+			ULevel* CurrentLevel = World->GetCurrentLevel();
+			if (CurrentLevel && !CurrentLevel->GetWorld()->GetSubsystem<UActorPartitionSubsystem>()->IsLevelPartition())
+			{
+				InOutSpawnParameters.ObjectFlags |= RF_Transient;
+			}
+		}
 	}
-}
 
-namespace UE::LevelSnapshots::Foliage::Private::Internal
-{
-	static UFoliageType* FindFoliageInfoFor(UHierarchicalInstancedStaticMeshComponent* Component)
+	void FFoliageSupport::PostRecreateActor(AActor* RecreatedActor)
 	{
-		AInstancedFoliageActor* Foliage = Cast<AInstancedFoliageActor>(Component->GetOwner());
-		if (!LIKELY(Foliage))
+		if (AInstancedFoliageActor* FoliageActor = Cast<AInstancedFoliageActor>(RecreatedActor))
 		{
-			return nullptr;
+			RecreatedActor->GetLevel()->InstancedFoliageActor = FoliageActor;
 		}
+	}
 
-		for (auto FoliageIt = Foliage->GetFoliageInfos().CreateConstIterator(); FoliageIt; ++FoliageIt)
+	void FFoliageSupport::PreRemoveActor(const FPreRemoveActorParams& Params)
+	{
+		if (AInstancedFoliageActor* FoliageActor = Cast<AInstancedFoliageActor>(Params.ActorToRemove))
 		{
-			if (FoliageIt->Value->Implementation->IsOwnedComponent(Component))
-			{
-				return FoliageIt->Key;
-			}
+			FoliageActor->GetLevel()->InstancedFoliageActor.Reset();
 		}
+	}
 
-		return nullptr;
+	void FFoliageSupport::PreRemoveComponent(UActorComponent* ComponentToRemove)
+	{
+		if (UFoliageType* FoliagType = FindFoliageInfoFor(ComponentToRemove))
+		{
+			AInstancedFoliageActor* FoliageActor = Cast<AInstancedFoliageActor>(ComponentToRemove->GetOwner());
+			FoliageActor->RemoveFoliageType(&FoliagType, 1);
+		}
 	}
 
-	static UFoliageType* FindFoliageInfoFor(UActorComponent* Component)
+	static IActorSnapshotFilter::FFilterResultData EvaluateSavedData(const FWorldSnapshotData& WorldData)
 	{
-		if (UHierarchicalInstancedStaticMeshComponent* Comp = Cast<UHierarchicalInstancedStaticMeshComponent>(Component))
+		const int32 SavedVersion = WorldData.SnapshotVersionInfo.GetSnapshotCustomVersion();
+		const bool bHasDataFromBefore5dot2 = SavedVersion < FSnapshotCustomVersion::FoliageTypesUnreadable;
+		const bool bHasBrokenData = bHasDataFromBefore5dot2 && FSnapshotCustomVersion::CustomSubobjectSoftObjectPathRefactor <= SavedVersion;
+		
+		if (bHasBrokenData)
 		{
-			return FindFoliageInfoFor(Comp);
+			return IActorSnapshotFilter::FFilterResultData{
+				IActorSnapshotFilter::EFilterResult::Disallow,
+				LOCTEXT("CannotModify.5dot1", "Snapshot was captured in 5.1. This data is unusable.")
+			};
 		}
-		return nullptr;
+		
+		if (bHasDataFromBefore5dot2 && !CVarAllowFoliageDataPre5dot1.GetValueOnGameThread())
+		{
+			return IActorSnapshotFilter::FFilterResultData{
+				IActorSnapshotFilter::EFilterResult::Disallow,
+				LOCTEXT("CannotModify.UseConsoleVariable", "Use \"LevelSnapshots.AllowFoliageDataBefore5dot1 true\" to enable restoring foliage data from before 5.1 (may crash).")
+			};
+		}
+
+		return IActorSnapshotFilter::FFilterResultData{ IActorSnapshotFilter::EFilterResult::Allow };
 	}
-}
 
-void UE::LevelSnapshots::Foliage::Private::FFoliageSupport::PreRemoveComponent(UActorComponent* ComponentToRemove)
-{
-	if (UFoliageType* FoliagType = Internal::FindFoliageInfoFor(ComponentToRemove))
+	IActorSnapshotFilter::FFilterResultData FFoliageSupport::CanModifyMatchedActor(const FCanModifyMatchedActorParams& Params)
+	{
+		return Params.MatchedEditorWorldActor.IsA<AInstancedFoliageActor>()
+			? EvaluateSavedData(Params.WorldData)
+			: EFilterResult::DoNotCare;
+	}
+
+	IActorSnapshotFilter::FFilterResultData FFoliageSupport::CanRecreateMissingActor(const FCanRecreateActorParams& Params)
 	{
-		AInstancedFoliageActor* FoliageActor = Cast<AInstancedFoliageActor>(ComponentToRemove->GetOwner());
-		FoliageActor->RemoveFoliageType(&FoliagType, 1);
+		return Params.Class && Params.Class->IsChildOf(AInstancedFoliageActor::StaticClass())
+			? EvaluateSavedData(Params.WorldData)
+			: EFilterResult::DoNotCare;
 	}
 }
+#undef LOCTEXT_NAMESPACE
\ No newline at end of file
diff --git a/Engine/Plugins/VirtualProduction/LevelSnapshots/Source/FoliageSupport/Private/FoliageSupport/FoliageSupport.h b/Engine/Plugins/VirtualProduction/LevelSnapshots/Source/FoliageSupport/Private/FoliageSupport/FoliageSupport.h
--- a/Engine/Plugins/VirtualProduction/LevelSnapshots/Source/FoliageSupport/Private/FoliageSupport/FoliageSupport.h
+++ b/Engine/Plugins/VirtualProduction/LevelSnapshots/Source/FoliageSupport/Private/FoliageSupport/FoliageSupport.h
@@ -2,14 +2,15 @@
 
 #pragma once
 
-#include "CoreMinimal.h"
 #include "ILevelSnapshotsModule.h"
-#include "FoliageSupport/InstancedFoliageActorData.h"
 #include "Interfaces/ISnapshotRestorabilityOverrider.h"
 #include "Interfaces/IRestorationListener.h"
 #include "Interfaces/ICustomObjectSnapshotSerializer.h"
+#include "UObject/WeakObjectPtr.h"
+#include "UObject/WeakObjectPtrTemplates.h"
 
 class AInstancedFoliageActor;
+class UObject;
 
 namespace UE::LevelSnapshots::Foliage::Private
 {
@@ -17,10 +18,14 @@ namespace UE::LevelSnapshots::Foliage::Private
 		:
 		public ISnapshotRestorabilityOverrider,
 		public ICustomObjectSnapshotSerializer,
-		public IRestorationListener
+		public IRestorationListener,
+		public IActorSnapshotFilter // For not allowing SkipSubobjectReferencesWhenSerializingArchetypesRefactor < version < FoliageTypesUnreadable
 	{
 		/** Just used to verify snapshots is in a valid state */
 		TWeakObjectPtr<AInstancedFoliageActor> CurrentFoliageActor;
+
+		/** Tracks modified foliage actors since they must be refreshed AFTER the snapshot has serialized all properties (in PostApplySnapshot). */
+		TArray<TWeakObjectPtr<AInstancedFoliageActor>> ActorsNeedingRefresh;
 		
 	public:
 
@@ -41,6 +46,7 @@ namespace UE::LevelSnapshots::Foliage::Private
 		//~ End ICustomObjectSnapshotSerializer Interface
 		
 		//~ Begin IRestorationListener Interface
+		virtual void PostApplySnapshot(const FApplySnapshotParams& Params) override;
 		virtual void PostApplySnapshotToActor(const FApplySnapshotToActorParams& Params) override;
 		virtual void PreApplySnapshotProperties(const FApplySnapshotPropertiesParams& Params) override;
 		virtual void PostApplySnapshotProperties(const FApplySnapshotPropertiesParams& Params) override;
@@ -49,6 +55,11 @@ namespace UE::LevelSnapshots::Foliage::Private
 		virtual void PreRemoveActor(const FPreRemoveActorParams& Params) override;
 		virtual void PreRemoveComponent(UActorComponent* ComponentToRemove) override;
 		//~ End IRestorationListener Interface
+
+		//~ Begin IActorSnapshotFilter Interface
+		virtual FFilterResultData CanModifyMatchedActor(const FCanModifyMatchedActorParams& Params) override;
+		virtual FFilterResultData CanRecreateMissingActor(const FCanRecreateActorParams& Params) override;
+		//~ End IActorSnapshotFilter Interface
 	};
 }
 
diff --git a/Engine/Plugins/VirtualProduction/LevelSnapshots/Source/FoliageSupport/Private/FoliageSupport/InstancedFoliageActorData.cpp b/Engine/Plugins/VirtualProduction/LevelSnapshots/Source/FoliageSupport/Private/FoliageSupport/InstancedFoliageActorData.cpp
deleted file mode 100644
diff --git a/Engine/Plugins/VirtualProduction/LevelSnapshots/Source/FoliageSupport/Private/FoliageSupport/InstancedFoliageActorData.h b/Engine/Plugins/VirtualProduction/LevelSnapshots/Source/FoliageSupport/Private/FoliageSupport/InstancedFoliageActorData.h
deleted file mode 100644
diff --git a/Engine/Plugins/VirtualProduction/LevelSnapshots/Source/FoliageSupport/Private/FoliageSupport/SubobjectFoliageInfoData.cpp b/Engine/Plugins/VirtualProduction/LevelSnapshots/Source/FoliageSupport/Private/FoliageSupport/SubobjectFoliageInfoData.cpp
deleted file mode 100644
diff --git a/Engine/Plugins/VirtualProduction/LevelSnapshots/Source/FoliageSupport/Private/FoliageSupport/SubobjectFoliageInfoData.h b/Engine/Plugins/VirtualProduction/LevelSnapshots/Source/FoliageSupport/Private/FoliageSupport/SubobjectFoliageInfoData.h
deleted file mode 100644
diff --git a/Engine/Plugins/VirtualProduction/LevelSnapshots/Source/LevelSnapshots/Private/Data/LevelSnapshot.cpp b/Engine/Plugins/VirtualProduction/LevelSnapshots/Source/LevelSnapshots/Private/Data/LevelSnapshot.cpp
--- a/Engine/Plugins/VirtualProduction/LevelSnapshots/Source/LevelSnapshots/Private/Data/LevelSnapshot.cpp
+++ b/Engine/Plugins/VirtualProduction/LevelSnapshots/Source/LevelSnapshots/Private/Data/LevelSnapshot.cpp
@@ -229,11 +229,16 @@ void ULevelSnapshot::DiffWorld(UWorld* World, FActorConsumer HandleMatchedActor,
 					continue;
 				}
 			
+				FText ExclusionReason;
 				UE::LevelSnapshots::Private::Internal::ConditionalBreakOnActor(DebugActorName, ActorInLevel);
-				if (!SerializedData.HasMatchingSavedActor(ActorInLevel) && UE::LevelSnapshots::Restorability::ShouldConsiderNewActorForRemoval(ActorInLevel))
+				if (!SerializedData.HasMatchingSavedActor(ActorInLevel) && UE::LevelSnapshots::Restorability::ShouldConsiderNewActorForRemoval(ActorInLevel, &ExclusionReason))
 				{
 					HandleAddedActor.Execute(ActorInLevel);
 				}
+				else
+				{
+					UE_CLOG(!ExclusionReason.IsEmpty(), LogLevelSnapshots, Warning, TEXT("Actor %s was not considered for removal. Reason: %s"), *ActorInLevel->GetPathName(), *ExclusionReason.ToString());
+				}
 			}
 		}
 	}
@@ -265,17 +270,20 @@ void ULevelSnapshot::DiffWorld(UWorld* World, FActorConsumer HandleMatchedActor,
 				return;
 			}
 			
-			UObject* ResolvedActor = OriginalActorPath.ResolveObject();
+			AActor* ResolvedActor = Cast<AActor>(OriginalActorPath.ResolveObject());
 			// OriginalActorPath may still resolve to a live actor if it was just removed. We need to check the ULevel::Actors to see whether it was removed.
 			const bool bWasRemovedFromWorld = ResolvedActor == nullptr || !AllActors.Contains(Cast<AActor>(ResolvedActor));
 			if (bWasRemovedFromWorld)
 			{
+				FText ExclusionReason;
 				auto GetDeserializedActorFunc = [this, &OriginalActorPath](){ return GetDeserializedActor(OriginalActorPath); };
 				const UE::LevelSnapshots::FCanRecreateActorParams Params{ World, ActorClass, OriginalActorPath, SavedData.SerializedActorData.GetObjectFlags(), SerializedData, GetDeserializedActorFunc };
-				if (UE::LevelSnapshots::Restorability::ShouldConsiderRemovedActorForRecreation(Params))
+				if (UE::LevelSnapshots::Restorability::ShouldConsiderRemovedActorForRecreation(Params, &ExclusionReason))
 				{
 					HandleRemovedActor.Execute(OriginalActorPath);
 				}
+				
+				UE_CLOG(!ExclusionReason.IsEmpty(), LogLevelSnapshots, Warning, TEXT("Actor %s was not considered for recreation. Reason: %s"), *OriginalActorPath.ToString(), *ExclusionReason.ToString());
 				return;
 			}
 			if (Settings->SkippedClasses.SkippedClasses.Contains(ActorClass))
@@ -283,13 +291,14 @@ void ULevelSnapshot::DiffWorld(UWorld* World, FActorConsumer HandleMatchedActor,
 				return;
 			}
 			
+			FText ExclusionReason;
 			// Possible scenario: Right-click actor > Replace Selected Actors with; deletes the original and replaces it with new actor.
 			if (ResolvedActor->GetClass() != ActorClass)
 			{
 				HandleRemovedActor.Execute(OriginalActorPath);
 				HandleAddedActor.Execute(Cast<AActor>(ResolvedActor));
 			}
-			else
+			else if (UE::LevelSnapshots::Restorability::ShouldConsiderMatchedActorForRestoration({ *ResolvedActor, SerializedData }, &ExclusionReason))
 			{
 				const FScopedLogItem Log = SortedItems.AddScopedLogItem(OriginalActorPath.ToString());
 				UE::LevelSnapshots::Private::Internal::ConditionalBreakOnActor(DebugActorName, OriginalActorPath);
@@ -298,6 +307,8 @@ void ULevelSnapshot::DiffWorld(UWorld* World, FActorConsumer HandleMatchedActor,
 				UE_LOG(LogLevelSnapshots, VeryVerbose, TEXT("Matched actor %s"), *OriginalActorPath.ToString());
 				HandleMatchedActor.Execute(Cast<AActor>(ResolvedActor));
 			}
+			
+			UE_CLOG(!ExclusionReason.IsEmpty(), LogLevelSnapshots, Warning, TEXT("Actor %s was skipped for modification. Reason: %s"), *OriginalActorPath.ToString(), *ExclusionReason.ToString());
 		});
 	}
 }
diff --git a/Engine/Plugins/VirtualProduction/LevelSnapshots/Source/LevelSnapshots/Private/LevelSnapshotsModule.cpp b/Engine/Plugins/VirtualProduction/LevelSnapshots/Source/LevelSnapshots/Private/LevelSnapshotsModule.cpp
--- a/Engine/Plugins/VirtualProduction/LevelSnapshots/Source/LevelSnapshots/Private/LevelSnapshotsModule.cpp
+++ b/Engine/Plugins/VirtualProduction/LevelSnapshots/Source/LevelSnapshots/Private/LevelSnapshotsModule.cpp
@@ -17,9 +17,31 @@
 #include "ISettingsModule.h"
 #endif
 
-namespace UE::LevelSnapshots::Private::Internal
+namespace UE::LevelSnapshots::Private
 {
-	
+	/** @param TPredicate signature IActorSnapshotFilter::FFilterResultData(IActorSnapshotFilter&) */
+	template<typename TPredicate>
+	static bool CanRestoreActor(const TArray<TSharedRef<IActorSnapshotFilter>>& GlobalFilters, FText* ExclusionReason, TPredicate Predicate)
+	{
+		for (const TSharedRef<IActorSnapshotFilter>& Filter : GlobalFilters)
+		{
+			const IActorSnapshotFilter::FFilterResultData Result = Predicate(Filter.Get());
+			switch (Result.InclusionResult)
+			{
+			case IActorSnapshotFilter::EFilterResult::Allow: return true;
+			case IActorSnapshotFilter::EFilterResult::DoNotCare: continue;
+			case IActorSnapshotFilter::EFilterResult::Disallow: 
+				if (ExclusionReason && !Result.ExclusionReason.IsEmpty())
+				{
+					*ExclusionReason = Result.ExclusionReason;
+				}
+				return false;
+			default:
+				checkNoEntry();
+			}
+		};
+		return true;
+	}
 }
 
 UE::LevelSnapshots::Private::FLevelSnapshotsModule& UE::LevelSnapshots::Private::FLevelSnapshotsModule::GetInternalModuleInstance()
@@ -27,7 +49,7 @@ UE::LevelSnapshots::Private::FLevelSnapshotsModule& UE::LevelSnapshots::Private:
 	static UE::LevelSnapshots::Private::FLevelSnapshotsModule& ModuleInstance = *[]() -> UE::LevelSnapshots::Private::FLevelSnapshotsModule*
 	{
 		UE_CLOG(!FModuleManager::Get().IsModuleLoaded("LevelSnapshots"), LogLevelSnapshots, Fatal, TEXT("You called GetInternalModuleInstance before the module was initialised."));
-		return &FModuleManager::GetModuleChecked<UE::LevelSnapshots::Private::FLevelSnapshotsModule>("LevelSnapshots");
+		return &FModuleManager::GetModuleChecked<FLevelSnapshotsModule>("LevelSnapshots");
 	}();
 	return ModuleInstance;
 }
@@ -327,40 +349,31 @@ TSharedPtr<UE::LevelSnapshots::ICustomObjectSnapshotSerializer> UE::LevelSnapsho
 	return nullptr;
 }
 
-bool UE::LevelSnapshots::Private::FLevelSnapshotsModule::CanRecreateActor(const FCanRecreateActorParams& Params) const
+bool UE::LevelSnapshots::Private::FLevelSnapshotsModule::CanModifyMatchedActor(const FCanModifyMatchedActorParams& Params, FText* ExclusionReason) const
 {
-	for (const TSharedRef<IActorSnapshotFilter>& Filter : GlobalFilters)
-	{
-		const IActorSnapshotFilter::EFilterResult Result = Filter->CanRecreateActor(Params);
-		switch (Result)
-		{
-		case IActorSnapshotFilter::EFilterResult::Allow: return true;
-		case IActorSnapshotFilter::EFilterResult::DoNotCare: continue;
-		case IActorSnapshotFilter::EFilterResult::Disallow: return false;
-		default:
-			checkNoEntry();
-		}
-	};
-
-	return true;
+	return CanRestoreActor(
+		GlobalFilters,
+		ExclusionReason,
+		[&Params](IActorSnapshotFilter& Filter){ return Filter.CanModifyMatchedActor(Params); }
+		);
 }
 
-bool UE::LevelSnapshots::Private::FLevelSnapshotsModule::CanDeleteActor(const AActor* EditorActor) const
+bool UE::LevelSnapshots::Private::FLevelSnapshotsModule::CanRecreateActor(const FCanRecreateActorParams& Params, FText* ExclusionReason) const
 {
-	for (const TSharedRef<IActorSnapshotFilter>& Filter : GlobalFilters)
-	{
-		const IActorSnapshotFilter::EFilterResult Result = Filter->CanDeleteActor(EditorActor);
-		switch (Result)
-		{
-		case IActorSnapshotFilter::EFilterResult::Allow: return true;
-		case IActorSnapshotFilter::EFilterResult::DoNotCare: continue;
-		case IActorSnapshotFilter::EFilterResult::Disallow: return false;
-		default:
-			checkNoEntry();
-		}
-	};
+	return CanRestoreActor(
+		GlobalFilters,
+		ExclusionReason,
+		[&Params](IActorSnapshotFilter& Filter){ return Filter.CanRecreateMissingActor(Params); }
+		);
+}
 
-	return true;
+bool UE::LevelSnapshots::Private::FLevelSnapshotsModule::CanDeleteActor(const AActor* EditorActor, FText* ExclusionReason) const
+{
+	return CanRestoreActor(
+		GlobalFilters,
+		ExclusionReason,
+		[EditorActor](IActorSnapshotFilter& Filter){ return Filter.CanDeleteNewActor(EditorActor); }
+		);
 }
 
 void UE::LevelSnapshots::Private::FLevelSnapshotsModule::AddCanTakeSnapshotDelegate(FName DelegateName, FCanTakeSnapshot Delegate)
diff --git a/Engine/Plugins/VirtualProduction/LevelSnapshots/Source/LevelSnapshots/Private/LevelSnapshotsModule.h b/Engine/Plugins/VirtualProduction/LevelSnapshots/Source/LevelSnapshots/Private/LevelSnapshotsModule.h
--- a/Engine/Plugins/VirtualProduction/LevelSnapshots/Source/LevelSnapshots/Private/LevelSnapshotsModule.h
+++ b/Engine/Plugins/VirtualProduction/LevelSnapshots/Source/LevelSnapshots/Private/LevelSnapshotsModule.h
@@ -64,8 +64,9 @@ namespace UE::LevelSnapshots::Private
 
 		TSharedPtr<ICustomObjectSnapshotSerializer> GetCustomSerializerForClass(UClass* Class) const;
 
-		bool CanRecreateActor(const FCanRecreateActorParams& Params) const;
-		bool CanDeleteActor(const AActor* EditorActor) const;
+		bool CanModifyMatchedActor(const FCanModifyMatchedActorParams& Params, FText* ExclusionReason = nullptr) const;
+		bool CanRecreateActor(const FCanRecreateActorParams& Params, FText* ExclusionReason = nullptr) const;
+		bool CanDeleteActor(const AActor* EditorActor, FText* ExclusionReason = nullptr) const;
 
 		virtual void AddCanTakeSnapshotDelegate(FName DelegateName, FCanTakeSnapshot Delegate) override;
 		virtual void RemoveCanTakeSnapshotDelegate(FName DelegateName) override;
diff --git a/Engine/Plugins/VirtualProduction/LevelSnapshots/Source/LevelSnapshots/Private/Restorability/ActorGroupRestoration.cpp b/Engine/Plugins/VirtualProduction/LevelSnapshots/Source/LevelSnapshots/Private/Restorability/ActorGroupRestoration.cpp
--- a/Engine/Plugins/VirtualProduction/LevelSnapshots/Source/LevelSnapshots/Private/Restorability/ActorGroupRestoration.cpp
+++ b/Engine/Plugins/VirtualProduction/LevelSnapshots/Source/LevelSnapshots/Private/Restorability/ActorGroupRestoration.cpp
@@ -224,7 +224,7 @@ namespace UE::LevelSnapshots::Private::ActorGroupRestoration
 			}
 		}
 
-		virtual EFilterResult CanRecreateActor(const FCanRecreateActorParams& Params) override
+		virtual FFilterResultData CanRecreateMissingActor(const FCanRecreateActorParams& Params) override
 		{
 			if (Params.Class == AGroupActor::StaticClass())
 			{
diff --git a/Engine/Plugins/VirtualProduction/LevelSnapshots/Source/LevelSnapshots/Private/Restorability/SnapshotRestorability.cpp b/Engine/Plugins/VirtualProduction/LevelSnapshots/Source/LevelSnapshots/Private/Restorability/SnapshotRestorability.cpp
--- a/Engine/Plugins/VirtualProduction/LevelSnapshots/Source/LevelSnapshots/Private/Restorability/SnapshotRestorability.cpp
+++ b/Engine/Plugins/VirtualProduction/LevelSnapshots/Source/LevelSnapshots/Private/Restorability/SnapshotRestorability.cpp
@@ -214,16 +214,21 @@ bool UE::LevelSnapshots::Restorability::IsPropertyDesirableForCapture(const FPro
 bool UE::LevelSnapshots::Restorability::IsPropertyExplicitlyUnsupportedForCapture(const FProperty* Property)
 {
 	SCOPED_SNAPSHOT_CORE_TRACE(IsPropertyExplicitlyUnsupportedForCapture);
-	return UE::LevelSnapshots::Private::FLevelSnapshotsModule::GetInternalModuleInstance().IsPropertyExplicitlyUnsupported(Property);
+	return LevelSnapshots::Private::FLevelSnapshotsModule::GetInternalModuleInstance().IsPropertyExplicitlyUnsupported(Property);
 }
 
 bool UE::LevelSnapshots::Restorability::IsPropertyExplicitlySupportedForCapture(const FProperty* Property)
 {
 	SCOPED_SNAPSHOT_CORE_TRACE(IsPropertyExplicitlySupportedForCapture);
-	return UE::LevelSnapshots::Private::FLevelSnapshotsModule::GetInternalModuleInstance().IsPropertyExplicitlySupported(Property);
+	return LevelSnapshots::Private::FLevelSnapshotsModule::GetInternalModuleInstance().IsPropertyExplicitlySupported(Property);
 }
 
-bool UE::LevelSnapshots::Restorability::ShouldConsiderRemovedActorForRecreation(const FCanRecreateActorParams& Params)
+bool UE::LevelSnapshots::Restorability::ShouldConsiderMatchedActorForRestoration(const FCanModifyMatchedActorParams& Params, FText* ExclusionReason)
+{
+	return LevelSnapshots::Private::FLevelSnapshotsModule::GetInternalModuleInstance().CanModifyMatchedActor(Params, ExclusionReason);
+}
+
+bool UE::LevelSnapshots::Restorability::ShouldConsiderRemovedActorForRecreation(const FCanRecreateActorParams& Params, FText* ExclusionReason)
 {
 	// Two reasons RF_Transactional is required:
 	// 1. It will not show up in Multiuser otherwise
@@ -234,14 +239,14 @@ bool UE::LevelSnapshots::Restorability::ShouldConsiderRemovedActorForRecreation(
 		|| Params.WorldData.SnapshotVersionInfo.GetSnapshotCustomVersion() >= FSnapshotCustomVersion::ClassArchetypeRefactor;
 	
 	return (bTransacts || !bWereFlagsSaved)
-		&& LevelSnapshots::Private::FLevelSnapshotsModule::GetInternalModuleInstance().CanRecreateActor(Params);
+		&& LevelSnapshots::Private::FLevelSnapshotsModule::GetInternalModuleInstance().CanRecreateActor(Params, ExclusionReason);
 }
 
-bool UE::LevelSnapshots::Restorability::ShouldConsiderNewActorForRemoval(const AActor* Actor)
+bool UE::LevelSnapshots::Restorability::ShouldConsiderNewActorForRemoval(const AActor* Actor, FText* ExclusionReason)
 {
 	return Private::Internal::DoesActorHaveSupportedClassForRemoving(Actor)
 		&& IsActorDesirableForCapture(Actor)
-		&& LevelSnapshots::Private::FLevelSnapshotsModule::GetInternalModuleInstance().CanDeleteActor(Actor);
+		&& LevelSnapshots::Private::FLevelSnapshotsModule::GetInternalModuleInstance().CanDeleteActor(Actor, ExclusionReason);
 }
 
 bool UE::LevelSnapshots::Restorability::IsRestorableProperty(const FProperty* LeafProperty)
diff --git a/Engine/Plugins/VirtualProduction/LevelSnapshots/Source/LevelSnapshots/Public/Data/SnapshotCustomVersion.h b/Engine/Plugins/VirtualProduction/LevelSnapshots/Source/LevelSnapshots/Public/Data/SnapshotCustomVersion.h
--- a/Engine/Plugins/VirtualProduction/LevelSnapshots/Source/LevelSnapshots/Public/Data/SnapshotCustomVersion.h
+++ b/Engine/Plugins/VirtualProduction/LevelSnapshots/Source/LevelSnapshots/Public/Data/SnapshotCustomVersion.h
@@ -45,6 +45,20 @@ namespace UE::LevelSnapshots
 			/** FActorSnapshotHash::Serialize was not implemented correctly (typo with TStructOpsTypeTraits using WithSerialize instead of WithSerializer). */
 			FixActorHashSerialize = 8,
 
+			/**
+			 * Change 21419024 of 16/08/2022 (which introduced CustomSubobjectSoftObjectPathRefactor above) causes
+			 * UE::LevelSnapshots::Foliage::Private::FInstancedFoliageActorData to save foliage data in a bad format.
+			 *
+			 * Any foliage data from 16/08/2022 to 30/03/2023 is unusable.
+			 * That means that 5.0 is not affected, 5.1 is broken, and 5.2 is fixed.
+			 *
+			 * The issue is that FInstancedFoliageActorData::Save, called by FFoliageSupport::OnTakeSnapshot, called
+			 * AInstancedFoliageActor::ForEachFoliageInfo and saved the foliage types sequentially. The issue is that
+			 * we differentiate between instanced subobject and asset foliage types; they were saved interleaved without
+			 * tagging the data, which makes it impossible to read (e.g. subject1 > asset1 > asset2 > subject 2, etc.).
+			 */
+			FoliageTypesUnreadable = 9,
+
 			// -----<new versions can be added above this line>-------------------------------------------------
 			VersionPlusOne,
 			LatestVersion = VersionPlusOne - 1
diff --git a/Engine/Plugins/VirtualProduction/LevelSnapshots/Source/LevelSnapshots/Public/Restorability/Interfaces/IActorSnapshotFilter.h b/Engine/Plugins/VirtualProduction/LevelSnapshots/Source/LevelSnapshots/Public/Restorability/Interfaces/IActorSnapshotFilter.h
--- a/Engine/Plugins/VirtualProduction/LevelSnapshots/Source/LevelSnapshots/Public/Restorability/Interfaces/IActorSnapshotFilter.h
+++ b/Engine/Plugins/VirtualProduction/LevelSnapshots/Source/LevelSnapshots/Public/Restorability/Interfaces/IActorSnapshotFilter.h
@@ -13,19 +13,34 @@ struct FWorldSnapshotData;
 namespace UE::LevelSnapshots
 {
 	using FDeserializedActorGetter = TFunctionRef<TOptional<TNonNullPtr<AActor>>()>;
-	
+
+	/** Params for deciding whether an actor should be recreated */
 	struct FCanRecreateActorParams
 	{
+		/** The world the snapshot is being applied to */
 		UWorld* EditorWorld;
 
+		/** The actor's class */
 		UClass* Class;
+		/** The object path of the actor */
 		FSoftObjectPath ActorPath;
+		/** The flags of the actor */
 		EObjectFlags ObjectFlags;
+		/** All world data of the snapshot */
 		const FWorldSnapshotData& WorldData;
 		
 		/** If you need access to the actor's data, you can call this function. */
 		FDeserializedActorGetter DeserializeFromSnapshotFunc;
 	};
+
+	/** Params for deciding whether an existing actor should have its properties modified. */
+	struct FCanModifyMatchedActorParams
+	{
+		/** The actor in the editor world which is being considered for restoration. */
+		AActor& MatchedEditorWorldActor; 
+		/** All world data of the snapshot */
+		const FWorldSnapshotData& WorldData;
+	};
 	
 	/**
 	 * Can decide whether a given actor can be added to or removed from the world.
@@ -44,12 +59,45 @@ namespace UE::LevelSnapshots
 			Disallow
 		};
 
+		struct FFilterResultData
+		{
+			/** Determines what to do with the actor data */
+			EFilterResult InclusionResult;
+			/** The reason for InclusionResult == Disallow, if any. Displayed to the user. */
+			FText ExclusionReason;
+
+			FFilterResultData() = default;
+			FFilterResultData(EFilterResult Result, FText ExclusionReason = FText::GetEmpty()) : InclusionResult(Result), ExclusionReason(MoveTemp(ExclusionReason)) {}
+		};
+		
+		/** Called to check whether a matched actor is allowed to have any of its data restored. */
+		virtual FFilterResultData CanModifyMatchedActor(const FCanModifyMatchedActorParams& Params) { return EFilterResult::DoNotCare; }
+
+		/** Whether the actor is allowed to be recreated. Called when diffing to world. */
+		virtual FFilterResultData CanRecreateMissingActor(const FCanRecreateActorParams& Params)
+		{
+			PRAGMA_DISABLE_DEPRECATION_WARNINGS
+			return CanRecreateActor(Params);
+			PRAGMA_ENABLE_DEPRECATION_WARNINGS
+		}
+
+		/** Whether the actor is allowed to be removed. Called when diffing to world. */
+		virtual FFilterResultData CanDeleteNewActor(const AActor* EditorActor)
+		{
+			PRAGMA_DISABLE_DEPRECATION_WARNINGS
+			return CanDeleteActor(EditorActor);
+			PRAGMA_ENABLE_DEPRECATION_WARNINGS
+		}
+		
+		
 		/** Whether the actor is allowed to be recreated. Called when diffing to world. */
+		UE_DEPRECATED(5.2, "Use CanRecreateMissingActor instead.")
 		virtual EFilterResult CanRecreateActor(const FCanRecreateActorParams& Params) { return EFilterResult::DoNotCare; }
 
 		/** Whether the actor is allowed to be removed. Called when diffing to world. */
+		UE_DEPRECATED(5.2, "Use CanDeleteNewActor instead.")
 		virtual EFilterResult CanDeleteActor(const AActor* EditorActor) { return EFilterResult::DoNotCare; }
-
+		
 		
 		virtual ~IActorSnapshotFilter() = default;
 	};
diff --git a/Engine/Plugins/VirtualProduction/LevelSnapshots/Source/LevelSnapshots/Public/Restorability/SnapshotRestorability.h b/Engine/Plugins/VirtualProduction/LevelSnapshots/Source/LevelSnapshots/Public/Restorability/SnapshotRestorability.h
--- a/Engine/Plugins/VirtualProduction/LevelSnapshots/Source/LevelSnapshots/Public/Restorability/SnapshotRestorability.h
+++ b/Engine/Plugins/VirtualProduction/LevelSnapshots/Source/LevelSnapshots/Public/Restorability/SnapshotRestorability.h
@@ -4,18 +4,15 @@
 
 #include "CoreMinimal.h"
 
-
 class ULevelSnapshot;
 class AActor;
 class FProperty;
 class UActorComponent;
 
-namespace UE
+namespace UE::LevelSnapshots
 {
-	namespace LevelSnapshots
-	{
-		struct FCanRecreateActorParams;
-	}
+	struct FCanModifyMatchedActorParams;
+	struct FCanRecreateActorParams;
 }
 
 namespace UE::LevelSnapshots::Restorability
@@ -39,11 +36,13 @@ namespace UE::LevelSnapshots::Restorability
 	LEVELSNAPSHOTS_API bool IsPropertyExplicitlyUnsupportedForCapture(const FProperty* Property);
 	/* Is this property always captured by the snapshot system? */
 	LEVELSNAPSHOTS_API bool IsPropertyExplicitlySupportedForCapture(const FProperty* Property);
-
-	/** The actor actors was removed from the snapshot. Should we show it in the list of removed actors? */
-	LEVELSNAPSHOTS_API bool ShouldConsiderRemovedActorForRecreation(const FCanRecreateActorParams& Params);
+	
+	/** The actor exists in the world and the snapshot. Should we show it in the list of modified actors? */
+	LEVELSNAPSHOTS_API bool ShouldConsiderMatchedActorForRestoration(const FCanModifyMatchedActorParams& Params, FText* ExclusionReason = nullptr);
+	/** The actor was removed from the snapshot. Should we show it in the list of removed actors? */
+	LEVELSNAPSHOTS_API bool ShouldConsiderRemovedActorForRecreation(const FCanRecreateActorParams& Params, FText* ExclusionReason = nullptr);
 	/* The actor did not exist in the snapshot. Should we show it in the list of added actors? */
-	LEVELSNAPSHOTS_API bool ShouldConsiderNewActorForRemoval(const AActor* Actor);
+	LEVELSNAPSHOTS_API bool ShouldConsiderNewActorForRemoval(const AActor* Actor, FText* ExclusionReason = nullptr);
 	
 	/* Is this property captured by the snapshot system? */
 	LEVELSNAPSHOTS_API bool IsRestorableProperty(const FProperty* LeafProperty);
