Fix Level Snapshots unit tests causing assert in UObject::CheckDefaultSubobjectsInternal due to serializing subobject references into transient snapshot world actors

diff --git a/Engine/Plugins/VirtualProduction/LevelSnapshots/Source/LevelSnapshots/Private/Archive/ApplySnapshotToEditorArchive.cpp b/Engine/Plugins/VirtualProduction/LevelSnapshots/Source/LevelSnapshots/Private/Archive/ApplySnapshotToEditorArchive.cpp
--- a/Engine/Plugins/VirtualProduction/LevelSnapshots/Source/LevelSnapshots/Private/Archive/ApplySnapshotToEditorArchive.cpp
+++ b/Engine/Plugins/VirtualProduction/LevelSnapshots/Source/LevelSnapshots/Private/Archive/ApplySnapshotToEditorArchive.cpp
@@ -39,7 +39,7 @@ void UE::LevelSnapshots::Private::FApplySnapshotToEditorArchive::ApplyToExisting
 	if (ClassIndex) // Sometimes not available, e.g. for custom subobjects
 	{
 		const FSubobjectArchetypeFallbackInfo ClassFallbackInfo{ InOriginalObject->GetOuter(), InOriginalObject->GetFName(), InOriginalObject->GetFlags() };
-		SerializeSelectedClassDefaultsInto(InOriginalObject, InSharedData, *ClassIndex, Cache, ClassFallbackInfo, *Selection);
+		SerializeSelectedClassDefaultsIntoSubobject(InOriginalObject, InSharedData, *ClassIndex, Cache, ClassFallbackInfo, *Selection);
 	}
 	
 	// Step 2: Serialise  properties that were different from CDO at time of snapshotting and that are still different from CDO
diff --git a/Engine/Plugins/VirtualProduction/LevelSnapshots/Source/LevelSnapshots/Private/Data/Util/WorldData/AllocationActorUtil.cpp b/Engine/Plugins/VirtualProduction/LevelSnapshots/Source/LevelSnapshots/Private/Data/Util/WorldData/AllocationActorUtil.cpp
--- a/Engine/Plugins/VirtualProduction/LevelSnapshots/Source/LevelSnapshots/Private/Data/Util/WorldData/AllocationActorUtil.cpp
+++ b/Engine/Plugins/VirtualProduction/LevelSnapshots/Source/LevelSnapshots/Private/Data/Util/WorldData/AllocationActorUtil.cpp
@@ -224,6 +224,8 @@ namespace UE::LevelSnapshots::Private::Internal
 		AActor* PreallocatedActor = Preallocated.GetValue();
 		{
 			const FRestoreObjectScope FinishRestore = PreActorRestore_SnapshotWorld(PreallocatedActor, ActorData.CustomActorSerializationData, WorldData, Cache, ProcessObjectDependency, InLocalisationSnapshotPackage);
+			// CDO must be serialized into the actor 
+			SerializeClassDefaultsIntoActor(PreallocatedActor, WorldData, ActorData.ClassIndex, Cache);
 			FLoadSnapshotObjectArchive::ApplyToSnapshotWorldObject(ActorData.SerializedActorData, WorldData, Cache, PreallocatedActor, ProcessObjectDependency, InLocalisationSnapshotPackage);
 #if WITH_EDITOR
 			UE_LOG(LogLevelSnapshots, Verbose, TEXT("ActorLabel is \"%s\" for \"%s\" (editor object path \"%s\")"), *PreallocatedActor->GetActorLabel(), *PreallocatedActor->GetPathName(), *OriginalActorPath.ToString());
@@ -338,25 +340,24 @@ TOptional<TNonNullPtr<AActor>> UE::LevelSnapshots::Private::GetPreallocated(cons
 			return {};
 		}
 		
-		const TOptional<TNonNullPtr<AActor>> Template = GetActorClassDefault(WorldData, ActorData->ClassIndex, Cache);
-		if (!Template)
+		const FSoftClassPath ClassPath = GetClass(*ActorData, WorldData);
+		UClass* ActorClass = ClassPath.TryLoadClass<AActor>();
+		if (!ActorClass)
 		{
 			UE_LOG(LogLevelSnapshots, Warning, TEXT("Failed to allocate %s due to missing class."), *OriginalActorPath.ToString());
 			return {};
 		}
 
-		UClass* TargetClass = Template->GetClass();
 		FActorSpawnParameters SpawnParams;
-		SpawnParams.Template = Template.GetValue();
 		SpawnParams.NameMode = FActorSpawnParameters::ESpawnActorNameMode::Requested;
 		SpawnParams.OverrideLevel = Internal::DetermineMostSuitableLevel(OriginalActorPath, WorldData);
 		const FName ActorName = *ExtractLastSubobjectName(OriginalActorPath);
 		// Must still generate a unique name because the asset registry might have messed with our data during a rename and added a duplicate actor name...
 		SpawnParams.Name = FindObjectFast<UObject>(SpawnParams.OverrideLevel, ActorName)
-			? MakeUniqueObjectName(SpawnParams.OverrideLevel, TargetClass, ActorName)
+			? MakeUniqueObjectName(SpawnParams.OverrideLevel, ActorClass, ActorName)
 			: ActorName;
 		SpawnParams.bNoFail = true;
-		ActorCache.CachedSnapshotActor = SnapshotWorld->SpawnActor<AActor>(TargetClass, SpawnParams);
+		ActorCache.CachedSnapshotActor = SnapshotWorld->SpawnActor<AActor>(ActorClass, SpawnParams);
 
 		const bool bSpawnWasSuccess = ActorCache.CachedSnapshotActor.IsValid();
 		UE_CLOG(!bSpawnWasSuccess, LogLevelSnapshots, Error, TEXT("Failed to allocate %s due to missing class."), *OriginalActorPath.ToString());
diff --git a/Engine/Plugins/VirtualProduction/LevelSnapshots/Source/LevelSnapshots/Private/Data/Util/WorldData/ClassDataUtil.cpp b/Engine/Plugins/VirtualProduction/LevelSnapshots/Source/LevelSnapshots/Private/Data/Util/WorldData/ClassDataUtil.cpp
--- a/Engine/Plugins/VirtualProduction/LevelSnapshots/Source/LevelSnapshots/Private/Data/Util/WorldData/ClassDataUtil.cpp
+++ b/Engine/Plugins/VirtualProduction/LevelSnapshots/Source/LevelSnapshots/Private/Data/Util/WorldData/ClassDataUtil.cpp
@@ -19,42 +19,6 @@
 
 namespace UE::LevelSnapshots::Private
 {
-	TOptional<TNonNullPtr<AActor>> GetActorClassDefault(FWorldSnapshotData& WorldData, FClassDataIndex ClassIndex, FSnapshotDataCache& Cache)
-	{
-		check(WorldData.ClassData.IsValidIndex(ClassIndex));
-		
-		FClassSnapshotData& ClassData = WorldData.ClassData[ClassIndex];
-		UClass* Class = ClassData.ClassPath.TryLoadClass<AActor>();
-		if (!Class)
-		{
-			UE_LOG(LogLevelSnapshots, Error, TEXT("Unknown class %s. The snapshot is mostly likely referencing a class that was deleted."), *ClassData.ClassPath.ToString());
-			return {};
-		}
-		
-		if (ClassData.SnapshotFlags == ESnapshotClassFlags::SerializationSkippedArchetypeData
-			|| ClassData.SerializedData.Num() == 0)
-		{
-			return { CastChecked<AActor>(Class->GetDefaultObject()) };
-		}
-
-		TObjectPtr<UObject>& CachedCDO = Cache.ArchetypeObjects[ClassIndex];
-		if (IsValid(CachedCDO))
-		{
-			return { CastChecked<AActor>(CachedCDO) };
-		}
-		
-		UObject* CDO = NewObject<UObject>(
-			WorldData.SnapshotWorld->GetPackage(),
-			Class,
-			MakeUniqueObjectName(WorldData.SnapshotWorld->GetPackage(), Class)
-			// TODO: Set RF_Archetype
-			);
-		FApplyClassDefaulDataArchive::SerializeClassDefaultObject(ClassData, WorldData, CDO);
-
-		CachedCDO = CDO;
-		return { CastChecked<AActor>(CDO) };
-	}
-
 	TOptional<TNonNullPtr<UObject>> GetSubobjectArchetype(FWorldSnapshotData& WorldData, FClassDataIndex ClassIndex, FSnapshotDataCache& Cache, const FSubobjectArchetypeFallbackInfo& FallbackInfo)
 	{
 		if (Cache.ArchetypeObjects[ClassIndex])
@@ -78,7 +42,6 @@ namespace UE::LevelSnapshots::Private
 			GetTransientPackage(),
 			Class,
 			*FString("SnapshotArchetype_").Append(*MakeUniqueObjectName(GetTransientPackage(), Class).ToString())
-			// TODO: Set RF_Archetype
 			);
 		FApplyClassDefaulDataArchive::SerializeClassDefaultObject(*ArchetypeData.GetValue(), WorldData, Archetype);
 		Cache.ArchetypeObjects[ClassIndex] = Archetype;
@@ -131,11 +94,26 @@ namespace UE::LevelSnapshots::Private
 		return {};
 	}
 
+	void SerializeClassDefaultsIntoActor(AActor* Actor, FWorldSnapshotData& WorldData, FClassDataIndex ClassIndex, FSnapshotDataCache& Cache)
+	{
+		if (WorldData.ClassData.IsValidIndex(ClassIndex))
+		{
+			SerializeClassDefaultsInto(Actor, WorldData.ClassData[ClassIndex], WorldData);
+		}
+		else
+		{
+			UE_LOG(LogLevelSnapshots, Warning,
+				TEXT("No CDO saved for class '%s'. If you changed some class default values for this class, then the affected objects will have the latest values instead of the class defaults at the time the snapshot was taken. Should be nothing major to worry about."),
+				*Actor->GetClass()->GetName()
+				);
+		}
+	}
+
 	void SerializeClassDefaultsIntoSubobject(UObject* Object, FWorldSnapshotData& WorldData, FClassDataIndex ClassIndex, FSnapshotDataCache& Cache, const FSubobjectArchetypeFallbackInfo& FallbackInfo)
 	{
 		if (const TOptional<TNonNullPtr<FClassSnapshotData>> ArchetypeData = GetObjectArchetypeData(WorldData, ClassIndex, Cache, FallbackInfo))
 		{
-			SerializeClassDefaultsIntoSubobject(Object, *ArchetypeData, WorldData);
+			SerializeClassDefaultsInto(Object, *ArchetypeData, WorldData);
 		}
 		else
 		{
@@ -146,7 +124,7 @@ namespace UE::LevelSnapshots::Private
 		}
 	}
 
-	void SerializeClassDefaultsIntoSubobject(UObject* Object, FClassSnapshotData& DataToSerialize, FWorldSnapshotData& WorldData)
+	void SerializeClassDefaultsInto(UObject* Object, FClassSnapshotData& DataToSerialize, FWorldSnapshotData& WorldData)
 	{
 		const bool bHasData = DataToSerialize.SerializedData.Num() > 0 && (DataToSerialize.SnapshotFlags & ESnapshotClassFlags::SerializationSkippedArchetypeData) == ESnapshotClassFlags::NoFlags;
 		if (bHasData && !FLevelSnapshotsModule::GetInternalModuleInstance().ShouldSkipClassDefaultSerialization(Object->GetClass()))
@@ -155,7 +133,7 @@ namespace UE::LevelSnapshots::Private
 		}
 	}
 	
-	void SerializeSelectedClassDefaultsInto(UObject* Object, FWorldSnapshotData& WorldData, FClassDataIndex ClassIndex, FSnapshotDataCache& Cache, const FSubobjectArchetypeFallbackInfo& FallbackInfo, const FPropertySelection& PropertiesToRestore)
+	void SerializeSelectedClassDefaultsIntoSubobject(UObject* Object, FWorldSnapshotData& WorldData, FClassDataIndex ClassIndex, FSnapshotDataCache& Cache, const FSubobjectArchetypeFallbackInfo& FallbackInfo, const FPropertySelection& PropertiesToRestore)
 	{
 		if (const TOptional<TNonNullPtr<FClassSnapshotData>> ArchetypeData = GetObjectArchetypeData(WorldData, ClassIndex, Cache, FallbackInfo))
 		{
diff --git a/Engine/Plugins/VirtualProduction/LevelSnapshots/Source/LevelSnapshots/Private/Data/Util/WorldData/ClassDataUtil.h b/Engine/Plugins/VirtualProduction/LevelSnapshots/Source/LevelSnapshots/Private/Data/Util/WorldData/ClassDataUtil.h
--- a/Engine/Plugins/VirtualProduction/LevelSnapshots/Source/LevelSnapshots/Private/Data/Util/WorldData/ClassDataUtil.h
+++ b/Engine/Plugins/VirtualProduction/LevelSnapshots/Source/LevelSnapshots/Private/Data/Util/WorldData/ClassDataUtil.h
@@ -22,6 +22,12 @@ struct FWorldSnapshotData;
 
 namespace UE::LevelSnapshots::Private
 {
+	/**
+	 * Gives info to retrieve class defaults for a subobject when it was not saved.
+	 * 
+	 * Remember: Blueprints create new archetype objects for each component so the lookup depends on the context.
+	 * These are basically just the args for UObject::GetArchetypeFromRequiredInfo.
+	 */
 	struct FSubobjectArchetypeFallbackInfo
 	{
 		UObject* SubobjectOuter;
@@ -29,23 +35,27 @@ namespace UE::LevelSnapshots::Private
 		EObjectFlags SubobjectFlags;
 	};
 	
-	/** Gets the actor CDO that was saved when the snapshot was taken. */
-	TOptional<TNonNullPtr<AActor>> GetActorClassDefault(FWorldSnapshotData& WorldData, FClassDataIndex ClassIndex, FSnapshotDataCache& Cache);
 	/** Gets the archetype to use for constructing a subobject.*/
 	TOptional<TNonNullPtr<UObject>> GetSubobjectArchetype(FWorldSnapshotData& WorldData, FClassDataIndex ClassIndex, FSnapshotDataCache& Cache, const FSubobjectArchetypeFallbackInfo& FallbackInfo);
 	
 	/** Gets an archetype that was saved when the snapshot was taken. Each subobject in actors can have their own archetypes. */
 	TOptional<TNonNullPtr<FClassSnapshotData>> GetObjectArchetypeData(FWorldSnapshotData& WorldData, FClassDataIndex ClassIndex, FSnapshotDataCache& Cache, const FSubobjectArchetypeFallbackInfo& FallbackInfo);
-	
+
+	/** Retrieves and serializes class default data into the given actor, if any was saved in the snapshot. */
+	void SerializeClassDefaultsIntoActor(AActor* Actor, FWorldSnapshotData& WorldData, FClassDataIndex ClassIndex, FSnapshotDataCache& Cache);
+	/** Retrieves and serializes class default data into the given actor, if any was saved in the snapshot. */
 	void SerializeClassDefaultsIntoSubobject(UObject* Object, FWorldSnapshotData& WorldData, FClassDataIndex ClassIndex, FSnapshotDataCache& Cache, const FSubobjectArchetypeFallbackInfo& FallbackInfo);
-	void SerializeClassDefaultsIntoSubobject(UObject* Object, FClassSnapshotData& DataToSerialize, FWorldSnapshotData& WorldData);
+	/** Helper that serializes already retrieved class default data into an object. Checks whether this feature is enabled for this class. */
+	void SerializeClassDefaultsInto(UObject* Object, FClassSnapshotData& DataToSerialize, FWorldSnapshotData& WorldData);
 
-	void SerializeSelectedClassDefaultsInto(UObject* Object, FWorldSnapshotData& WorldData, FClassDataIndex ClassIndex, FSnapshotDataCache& Cache, const FSubobjectArchetypeFallbackInfo& FallbackInfo, const FPropertySelection& PropertiesToRestore);
+	/** Serializes the class default data into the given subobject but only those allowed by PropertiesToRestore. */
+	void SerializeSelectedClassDefaultsIntoSubobject(UObject* Object, FWorldSnapshotData& WorldData, FClassDataIndex ClassIndex, FSnapshotDataCache& Cache, const FSubobjectArchetypeFallbackInfo& FallbackInfo, const FPropertySelection& PropertiesToRestore);
 
+	/** Helper for getting an actor's class. */
 	FSoftClassPath GetClass(const FActorSnapshotData& Data, const FWorldSnapshotData& WorldData);
+	/** Helper for getting an actor's subobject's class. */
 	FSoftClassPath GetClass(const FSubobjectSnapshotData& Data, const FWorldSnapshotData& WorldData);
 	
-	
 	/** Adds a class default entry in the world data */
 	FClassDataIndex AddClassArchetype(FWorldSnapshotData& WorldData, UObject* SavedObject);
 }
\ No newline at end of file
diff --git a/Engine/Plugins/VirtualProduction/LevelSnapshots/Source/LevelSnapshots/Private/Data/Util/WorldData/WorldDataUtil.cpp b/Engine/Plugins/VirtualProduction/LevelSnapshots/Source/LevelSnapshots/Private/Data/Util/WorldData/WorldDataUtil.cpp
--- a/Engine/Plugins/VirtualProduction/LevelSnapshots/Source/LevelSnapshots/Private/Data/Util/WorldData/WorldDataUtil.cpp
+++ b/Engine/Plugins/VirtualProduction/LevelSnapshots/Source/LevelSnapshots/Private/Data/Util/WorldData/WorldDataUtil.cpp
@@ -388,8 +388,9 @@ namespace UE::LevelSnapshots::Private::Internal
 					continue;	
 				}
 
-				const TOptional<TNonNullPtr<AActor>> ActorCDO = GetActorClassDefault(WorldData, ActorSnapshot->ClassIndex, Cache);
-				if (!ActorCDO)
+				const FSoftClassPath ClassPath = GetClass(*ActorSnapshot, WorldData);
+				UClass* ActorClass = ClassPath.TryLoadClass<AActor>();
+				if (!ActorClass)
 				{
 					UE_LOG(LogLevelSnapshots, Warning, TEXT("Skipping recreation of %s due to missing class"), *OriginalRemovedActorPath.ToString());
 					continue;
@@ -418,7 +419,6 @@ namespace UE::LevelSnapshots::Private::Internal
 					ULevel* OverrideLevel = OwningLevelWorld->PersistentLevel;
 
 					const FName ActorFName = *ActorName;
-					UClass* ActorClass = ActorCDO->GetClass();
 					FActorSpawnParameters SpawnParameters;
 					// Cannot be overriden but we want to inform the subscribers
 					SpawnParameters.bNoFail = true;
@@ -427,7 +427,6 @@ namespace UE::LevelSnapshots::Private::Internal
 					SpawnParameters.NameMode = FActorSpawnParameters::ESpawnActorNameMode::Required_ErrorAndReturnNull;
 					SpawnParameters.bDeferConstruction = true;
 					// Overriable properties
-					SpawnParameters.Template = ActorCDO.GetValue();
 					SpawnParameters.SpawnCollisionHandlingOverride = ESpawnActorCollisionHandlingMethod::AlwaysSpawn;
 					SpawnParameters.ObjectFlags = ActorSnapshot->SerializedActorData.GetObjectFlags()
 						// Was wrongly saved for some very old snapshots - if this causes trouble at some point ... consider just deprecating the old snapshots
