Fix Level Snapshots not restoring ActorLabel property

diff --git a/Engine/Plugins/VirtualProduction/LevelSnapshots/Source/LevelSnapshots/Private/Restorability/ActorLabelRestoration.cpp b/Engine/Plugins/VirtualProduction/LevelSnapshots/Source/LevelSnapshots/Private/Restorability/ActorLabelRestoration.cpp
new file mode 100644
--- /dev/null
+++ b/Engine/Plugins/VirtualProduction/LevelSnapshots/Source/LevelSnapshots/Private/Restorability/ActorLabelRestoration.cpp
@@ -0,0 +1,59 @@
+// Copyright Epic Games, Inc. All Rights Reserved.
+
+#include "ActorLabelRestoration.h"
+
+#include "GameFramework/Actor.h"
+#include "Misc/CoreDelegates.h"
+#include "Selection/PropertySelection.h"
+#include "UObject/UnrealType.h"
+
+namespace UE::LevelSnapshots::Private::ActorLabelRestoration
+{
+#if WITH_EDITOR // Actor labels are an editor only property
+	static void AddActorLabelSupport(ILevelSnapshotsModule& Module)
+	{
+		const FProperty* ActorLabel = AActor::StaticClass()->FindPropertyByName(TEXT("ActorLabel"));
+		if (ensure(ActorLabel))
+		{
+			Module.AddExplicitilySupportedProperties({ ActorLabel });
+		}
+	}
+
+	/** Does a cosmetic update to actors that have their ActorLabel property restored by refreshing the world outliner with the new label */
+	class FActorGroupRestoration : public IRestorationListener
+	{
+		FProperty* const ActorLabel = AActor::StaticClass()->FindPropertyByName(TEXT("ActorLabel"));
+	public:
+
+		FActorGroupRestoration()
+		{
+			ensure(ActorLabel);
+		}
+		
+		//~ Begin IRestorationListener Interface
+		virtual void PostApplySnapshotProperties(const FApplySnapshotPropertiesParams& Params) override
+		{
+			const bool bNeedsToUpdateActorLabel = Params.bWasRecreated
+				|| (ensureMsgf(Params.PropertySelection, TEXT("Supposed to be valid when !bWasRecreated"))
+					&& Params.PropertySelection.GetValue()->IsPropertySelected(nullptr, ActorLabel));
+			if (AActor* const RestoredActor = Cast<AActor>(Params.Object);
+				RestoredActor && bNeedsToUpdateActorLabel)
+			{
+				// Cannot call SetActorLabel directly because it checks whether the label has changed; the new value was already serialized into the actor
+				FPropertyChangedEvent PropertyEvent(ActorLabel);
+				RestoredActor->PostEditChangeProperty(PropertyEvent);
+				FCoreDelegates::OnActorLabelChanged.Broadcast(RestoredActor);
+			}
+		}
+		//~ End IRestorationListener Interface
+	};
+#endif
+	
+	void Register(FLevelSnapshotsModule& Module)
+	{
+#if WITH_EDITOR // Actor labels are an editor only property
+		AddActorLabelSupport(Module);
+		Module.RegisterRestorationListener(MakeShared<FActorGroupRestoration>());
+#endif
+	}
+}
\ No newline at end of file
diff --git a/Engine/Plugins/VirtualProduction/LevelSnapshots/Source/LevelSnapshots/Private/Restorability/ActorLabelRestoration.h b/Engine/Plugins/VirtualProduction/LevelSnapshots/Source/LevelSnapshots/Private/Restorability/ActorLabelRestoration.h
new file mode 100644
--- /dev/null
+++ b/Engine/Plugins/VirtualProduction/LevelSnapshots/Source/LevelSnapshots/Private/Restorability/ActorLabelRestoration.h
@@ -0,0 +1,10 @@
+// Copyright Epic Games, Inc. All Rights Reserved.
+
+#pragma once
+
+#include "LevelSnapshotsModule.h"
+
+namespace UE::LevelSnapshots::Private::ActorLabelRestoration
+{
+	void Register(FLevelSnapshotsModule& Module);
+}
diff --git a/Engine/Plugins/VirtualProduction/LevelSnapshots/Source/LevelSnapshots/Private/Restorability/EngineTypesRestorationFence.cpp b/Engine/Plugins/VirtualProduction/LevelSnapshots/Source/LevelSnapshots/Private/Restorability/EngineTypesRestorationFence.cpp
--- a/Engine/Plugins/VirtualProduction/LevelSnapshots/Source/LevelSnapshots/Private/Restorability/EngineTypesRestorationFence.cpp
+++ b/Engine/Plugins/VirtualProduction/LevelSnapshots/Source/LevelSnapshots/Private/Restorability/EngineTypesRestorationFence.cpp
@@ -3,6 +3,7 @@
 #include "EngineTypesRestorationFence.h"
 
 #include "ActorGroupRestoration.h"
+#include "ActorLabelRestoration.h"
 #include "AttachParentShowsTransformProperties.h"
 #include "LevelInstanceRestoration.h"
 #include "Restorability/CollisionRestoration.h"
@@ -128,5 +129,6 @@ namespace UE::LevelSnapshots::Private::EngineTypesRestorationFence
 		AttachParentShowsTransformPropertiesFix::Register(Module);
 		LevelInstanceRestoration::Register(Module);
 		MaterialInstanceSupport::Register(Module);
+		ActorLabelRestoration::Register(Module);
 	}
 }
diff --git a/Engine/Plugins/VirtualProduction/LevelSnapshots/Source/LevelSnapshots/Private/Tests/RegressionTests.cpp b/Engine/Plugins/VirtualProduction/LevelSnapshots/Source/LevelSnapshots/Private/Tests/RegressionTests.cpp
--- a/Engine/Plugins/VirtualProduction/LevelSnapshots/Source/LevelSnapshots/Private/Tests/RegressionTests.cpp
+++ b/Engine/Plugins/VirtualProduction/LevelSnapshots/Source/LevelSnapshots/Private/Tests/RegressionTests.cpp
@@ -8,6 +8,7 @@
 #include "Types/ActorWithReferencesInCDO.h"
 
 #include "Engine/World.h"
+#include "EngineUtils.h"
 #include "Misc/AutomationTest.h"
 
 // Bug fixes should generally be tested. Put tests for bug fixes here.
@@ -209,4 +210,46 @@ namespace UE::LevelSnapshots::Private::Tests
 
 		return true;
 	}
+
+	/**
+	 * Verify that the actor label property is restored
+	 */
+	IMPLEMENT_SIMPLE_AUTOMATION_TEST(FRestoreActorLabel, "VirtualProduction.LevelSnapshots.Snapshot.Regression.ActorLabelRestores", (EAutomationTestFlags::ApplicationContextMask | EAutomationTestFlags::EngineFilter));
+	bool FRestoreActorLabel::RunTest(const FString& Parameters)
+	{
+		AActor* KeepActor = nullptr;
+		AActor* RemovedActor = nullptr;
+		const FString KeepActorLabel = TEXT("KeepActor");
+		const FString RemovedActorLabel = TEXT("RemovedActor");
+	
+		FSnapshotTestRunner()
+			.ModifyWorld([&](UWorld* World)
+			{
+				KeepActor = World->SpawnActor<AActor>();
+				RemovedActor = World->SpawnActor<AActor>();;
+
+				KeepActor->SetActorLabel(KeepActorLabel);
+				RemovedActor->SetActorLabel(RemovedActorLabel);
+			})
+			.TakeSnapshot()
+			.ModifyWorld([&](UWorld* World)
+			{
+				KeepActor->SetActorLabel(TEXT("NewName"));
+				RemovedActor->Destroy();
+			})
+			.ApplySnapshot()
+			.ModifyWorld([&](UWorld* World)
+			{
+				TestEqual(TEXT("Modified actor label is restored"), KeepActor->GetActorLabel(), KeepActorLabel);
+
+				bool bRemovedActorHasCorrectLabel = false;
+				for (TActorIterator<AActor> ActorIt(World); ActorIt && !bRemovedActorHasCorrectLabel; ++ActorIt)
+				{
+					bRemovedActorHasCorrectLabel = ActorIt->GetActorLabel() == RemovedActorLabel;
+				}
+				TestTrue(TEXT("Removed actor label is restored"), bRemovedActorHasCorrectLabel);
+			});
+
+		return true;
+	}
 }
\ No newline at end of file
diff --git a/Engine/Plugins/VirtualProduction/LevelSnapshots/Source/LevelSnapshots/Private/Tests/Types/SnapshotTestActor.cpp b/Engine/Plugins/VirtualProduction/LevelSnapshots/Source/LevelSnapshots/Private/Tests/Types/SnapshotTestActor.cpp
--- a/Engine/Plugins/VirtualProduction/LevelSnapshots/Source/LevelSnapshots/Private/Tests/Types/SnapshotTestActor.cpp
+++ b/Engine/Plugins/VirtualProduction/LevelSnapshots/Source/LevelSnapshots/Private/Tests/Types/SnapshotTestActor.cpp
@@ -23,7 +23,15 @@ ASnapshotTestActor* ASnapshotTestActor::Spawn(UWorld* World, FName Name)
 	Params.bNoFail = true;
 	Params.SpawnCollisionHandlingOverride = ESpawnActorCollisionHandlingMethod::AlwaysSpawn;
 	Params.NameMode = FActorSpawnParameters::ESpawnActorNameMode::Requested;
-	return World->SpawnActor<ASnapshotTestActor>(Params);
+	ASnapshotTestActor* Result = World->SpawnActor<ASnapshotTestActor>(Params);
+
+	// Editor code sets an actor's label after adding it to the world. Note that GetActorLabel sets a label (see default param bCreateIfNone = true)
+	// Without this call, the actor would not have any label. When the unit test diffs properties, GetActorLabel will have been called, giving the actor a label.
+	// The transient snapshot actor would not have a label. This would generate an unexpected diff which a unit test is not even testing for
+	// Summary: Just give all actors a label so they do not generate false diffs.
+	Result->SetActorLabel(Name.ToString());
+	
+	return Result;
 }
 
 ASnapshotTestActor::ASnapshotTestActor()
diff --git a/Engine/Plugins/VirtualProduction/LevelSnapshots/Source/LevelSnapshots/Private/Tests/Types/SnapshotTestActor.h b/Engine/Plugins/VirtualProduction/LevelSnapshots/Source/LevelSnapshots/Private/Tests/Types/SnapshotTestActor.h
--- a/Engine/Plugins/VirtualProduction/LevelSnapshots/Source/LevelSnapshots/Private/Tests/Types/SnapshotTestActor.h
+++ b/Engine/Plugins/VirtualProduction/LevelSnapshots/Source/LevelSnapshots/Private/Tests/Types/SnapshotTestActor.h
@@ -74,7 +74,7 @@ class ASnapshotTestActor : public ACharacter
 	GENERATED_BODY()
 public:
 
-	static ASnapshotTestActor* Spawn(UWorld* World, FName Name = FName());
+	static ASnapshotTestActor* Spawn(UWorld* World, FName Name = TEXT("TestActor"));
 	
 	ASnapshotTestActor();
 
diff --git a/Engine/Plugins/VirtualProduction/LevelSnapshots/Source/LevelSnapshots/Public/Data/LevelSnapshot.h b/Engine/Plugins/VirtualProduction/LevelSnapshots/Source/LevelSnapshots/Public/Data/LevelSnapshot.h
--- a/Engine/Plugins/VirtualProduction/LevelSnapshots/Source/LevelSnapshots/Public/Data/LevelSnapshot.h
+++ b/Engine/Plugins/VirtualProduction/LevelSnapshots/Source/LevelSnapshots/Public/Data/LevelSnapshot.h
@@ -56,14 +56,15 @@ public:
 	 */
 	bool HasOriginalChangedPropertiesSinceSnapshotWasTaken(AActor* SnapshotActor, AActor* WorldActor);
 
-	/** Gets the display label of the path of the actor */
+	/** Gets the display label of the path of the actor. This is useful for the LS UI to display actor names, especially of removed actors. */
 	FString GetActorLabel(const FSoftObjectPath& OriginalActorPath) const;
 	
 	/** Given an actor path in the world, gets the equivalent actor from the snapshot. */
 	TOptional<TNonNullPtr<AActor>> GetDeserializedActor(const FSoftObjectPath& OriginalActorPath);
 	
-	
+	/** Gets the number of actors saved by this snapshot. */
 	int32 GetNumSavedActors() const;
+	
 	/**
 	 * Compares this snapshot to the world and calls the appropriate callbacks:
 	 * @param World to check in
