Reorganize Level Snapshots folder structure

diff --git a/Engine/Plugins/VirtualProduction/LevelSnapshots/Source/FoliageSupport/Private/FoliageSupport/Apply/ApplyDataUtils.cpp b/Engine/Plugins/VirtualProduction/LevelSnapshots/Source/FoliageSupport/Private/FoliageSupport/Apply/ApplyDataUtils.cpp
--- a/Engine/Plugins/VirtualProduction/LevelSnapshots/Source/FoliageSupport/Private/FoliageSupport/Apply/ApplyDataUtils.cpp
+++ b/Engine/Plugins/VirtualProduction/LevelSnapshots/Source/FoliageSupport/Private/FoliageSupport/Apply/ApplyDataUtils.cpp
@@ -3,7 +3,7 @@
 #include "ApplyDataUtils.h"
 
 #include "LevelSnapshotsLog.h"
-#include "Selection/PropertySelectionMap.h"
+#include "Filtering/PropertySelectionMap.h"
 
 #include "FoliageType.h"
 #include "InstancedFoliageActor.h"
diff --git a/Engine/Plugins/VirtualProduction/LevelSnapshots/Source/FoliageSupport/Private/FoliageSupport/Apply/ApplyData_5dot2.cpp b/Engine/Plugins/VirtualProduction/LevelSnapshots/Source/FoliageSupport/Private/FoliageSupport/Apply/ApplyData_5dot2.cpp
--- a/Engine/Plugins/VirtualProduction/LevelSnapshots/Source/FoliageSupport/Private/FoliageSupport/Apply/ApplyData_5dot2.cpp
+++ b/Engine/Plugins/VirtualProduction/LevelSnapshots/Source/FoliageSupport/Private/FoliageSupport/Apply/ApplyData_5dot2.cpp
@@ -3,9 +3,9 @@
 #include "FoliageSupport/Data/InstancedFoliageActorData.h"
 
 #include "ApplyDataUtils.h"
+#include "Filtering/PropertySelectionMap.h"
 #include "FoliageRestorationInfo.h"
 #include "LevelSnapshotsLog.h"
-#include "Selection/PropertySelectionMap.h"
 #include "SnapshotCustomVersion.h"
 
 #include "FoliageType.h"
diff --git a/Engine/Plugins/VirtualProduction/LevelSnapshots/Source/FoliageSupport/Private/FoliageSupport/Apply/ApplyData_Pre5dot1.cpp b/Engine/Plugins/VirtualProduction/LevelSnapshots/Source/FoliageSupport/Private/FoliageSupport/Apply/ApplyData_Pre5dot1.cpp
--- a/Engine/Plugins/VirtualProduction/LevelSnapshots/Source/FoliageSupport/Private/FoliageSupport/Apply/ApplyData_Pre5dot1.cpp
+++ b/Engine/Plugins/VirtualProduction/LevelSnapshots/Source/FoliageSupport/Private/FoliageSupport/Apply/ApplyData_Pre5dot1.cpp
@@ -3,10 +3,10 @@
 #include "FoliageSupport/Data/InstancedFoliageActorData.h"
 
 #include "ApplyDataUtils.h"
+#include "Filtering/PropertySelectionMap.h"
 #include "FoliageLevelSnapshotsConsoleVariables.h"
 #include "FoliageRestorationInfo.h"
 #include "LevelSnapshotsLog.h"
-#include "Selection/PropertySelectionMap.h"
 #include "SnapshotCustomVersion.h"
 
 #include "FoliageType.h"
diff --git a/Engine/Plugins/VirtualProduction/LevelSnapshots/Source/FoliageSupport/Private/FoliageSupport/Apply/FoliageRestorationInfo.cpp b/Engine/Plugins/VirtualProduction/LevelSnapshots/Source/FoliageSupport/Private/FoliageSupport/Apply/FoliageRestorationInfo.cpp
--- a/Engine/Plugins/VirtualProduction/LevelSnapshots/Source/FoliageSupport/Private/FoliageSupport/Apply/FoliageRestorationInfo.cpp
+++ b/Engine/Plugins/VirtualProduction/LevelSnapshots/Source/FoliageSupport/Private/FoliageSupport/Apply/FoliageRestorationInfo.cpp
@@ -2,11 +2,11 @@
 
 #include "FoliageSupport/Apply/FoliageRestorationInfo.h"
 
+#include "Filtering/AddedAndRemovedComponentInfo.h"
+#include "Filtering/RestorableObjectSelection.h"
+#include "Filtering/PropertySelectionMap.h"
 #include "FoliageSupport/Data/FoliageInfoData.h"
 #include "LevelSnapshotsLog.h"
-#include "Selection/AddedAndRemovedComponentInfo.h"
-#include "Selection/RestorableObjectSelection.h"
-#include "Selection/PropertySelectionMap.h"
 
 #include "InstancedFoliageActor.h"
 
diff --git a/Engine/Plugins/VirtualProduction/LevelSnapshots/Source/FoliageSupport/Private/FoliageSupport/Data/InstancedFoliageActorData.cpp b/Engine/Plugins/VirtualProduction/LevelSnapshots/Source/FoliageSupport/Private/FoliageSupport/Data/InstancedFoliageActorData.cpp
--- a/Engine/Plugins/VirtualProduction/LevelSnapshots/Source/FoliageSupport/Private/FoliageSupport/Data/InstancedFoliageActorData.cpp
+++ b/Engine/Plugins/VirtualProduction/LevelSnapshots/Source/FoliageSupport/Private/FoliageSupport/Data/InstancedFoliageActorData.cpp
@@ -2,7 +2,7 @@
 
 #include "FoliageSupport/Data/InstancedFoliageActorData.h"
 
-#include "Selection/PropertySelectionMap.h"
+#include "Filtering/PropertySelectionMap.h"
 #include "SnapshotCustomVersion.h"
 
 #include "FoliageType.h"
diff --git a/Engine/Plugins/VirtualProduction/LevelSnapshots/Source/FoliageSupport/Private/FoliageSupport/FoliageSupport.cpp b/Engine/Plugins/VirtualProduction/LevelSnapshots/Source/FoliageSupport/Private/FoliageSupport/FoliageSupport.cpp
--- a/Engine/Plugins/VirtualProduction/LevelSnapshots/Source/FoliageSupport/Private/FoliageSupport/FoliageSupport.cpp
+++ b/Engine/Plugins/VirtualProduction/LevelSnapshots/Source/FoliageSupport/Private/FoliageSupport/FoliageSupport.cpp
@@ -2,10 +2,10 @@
 
 #include "FoliageSupport/FoliageSupport.h"
 
+#include "Filtering/PropertySelection.h"
+#include "Filtering/PropertySelectionMap.h"
 #include "FoliageSupport/Data/InstancedFoliageActorData.h"
 #include "Params/ObjectSnapshotSerializationData.h"
-#include "Selection/PropertySelection.h"
-#include "Selection/PropertySelectionMap.h"
 #include "SnapshotCustomVersion.h"
 #include "WorldSnapshotData.h"
 
diff --git a/Engine/Plugins/VirtualProduction/LevelSnapshots/Source/LevelSnapshots/Private/Archive/ApplySnapshotToEditorArchive.cpp b/Engine/Plugins/VirtualProduction/LevelSnapshots/Source/LevelSnapshots/Private/Archive/ApplySnapshotToEditorArchive.cpp
--- a/Engine/Plugins/VirtualProduction/LevelSnapshots/Source/LevelSnapshots/Private/Archive/ApplySnapshotToEditorArchive.cpp
+++ b/Engine/Plugins/VirtualProduction/LevelSnapshots/Source/LevelSnapshots/Private/Archive/ApplySnapshotToEditorArchive.cpp
@@ -3,12 +3,12 @@
 #include "Archive/ApplySnapshotToEditorArchive.h"
 
 #include "Archive/ClassDefaults/ApplyClassDefaulDataArchive.h"
-#include "Data/RestorationEvents/ApplySnapshotPropertiesScope.h"
-#include "Data/Util/WorldData/ClassDataUtil.h"
-#include "Data/Util/WorldData/SnapshotObjectUtil.h"
+#include "Filtering/PropertySelection.h"
+#include "Data/Util/ClassDataUtil.h"
+#include "Data/Util/ObjectDependencyUtil.h"
 #include "Data/WorldSnapshotData.h"
 #include "LevelSnapshotsLog.h"
-#include "Selection/PropertySelection.h"
+#include "Restoration/Scope/ApplySnapshotPropertiesScope.h"
 
 #include "Serialization/ObjectReader.h"
 #include "UObject/UnrealType.h"
diff --git a/Engine/Plugins/VirtualProduction/LevelSnapshots/Source/LevelSnapshots/Private/Archive/ApplySnapshotToEditorArchive.h b/Engine/Plugins/VirtualProduction/LevelSnapshots/Source/LevelSnapshots/Private/Archive/ApplySnapshotToEditorArchive.h
--- a/Engine/Plugins/VirtualProduction/LevelSnapshots/Source/LevelSnapshots/Private/Archive/ApplySnapshotToEditorArchive.h
+++ b/Engine/Plugins/VirtualProduction/LevelSnapshots/Source/LevelSnapshots/Private/Archive/ApplySnapshotToEditorArchive.h
@@ -2,9 +2,8 @@
 
 #pragma once
 
-#include "CoreMinimal.h"
-#include "Selection/PropertySelection.h"
-#include "Selection/PropertySelectionMap.h"
+#include "Filtering/PropertySelection.h"
+#include "Filtering/PropertySelectionMap.h"
 #include "SnapshotArchive.h"
 #include "SnapshotUtilTypes.h"
 
diff --git a/Engine/Plugins/VirtualProduction/LevelSnapshots/Source/LevelSnapshots/Private/Archive/ClassDefaults/ApplyClassDefaulDataArchive.cpp b/Engine/Plugins/VirtualProduction/LevelSnapshots/Source/LevelSnapshots/Private/Archive/ClassDefaults/ApplyClassDefaulDataArchive.cpp
--- a/Engine/Plugins/VirtualProduction/LevelSnapshots/Source/LevelSnapshots/Private/Archive/ClassDefaults/ApplyClassDefaulDataArchive.cpp
+++ b/Engine/Plugins/VirtualProduction/LevelSnapshots/Source/LevelSnapshots/Private/Archive/ClassDefaults/ApplyClassDefaulDataArchive.cpp
@@ -2,12 +2,11 @@
 
 #include "Archive/ClassDefaults/ApplyClassDefaulDataArchive.h"
 
-#include "LevelSnapshotsLog.h"
+#include "Data/Util/ObjectDependencyUtil.h"
 #include "LevelSnapshotsModule.h"
 #include "ObjectSnapshotData.h"
 #include "SelectionSet.h"
 #include "WorldSnapshotData.h"
-#include "Util/WorldData/SnapshotObjectUtil.h"
 
 void UE::LevelSnapshots::Private::FApplyClassDefaulDataArchive::SerializeClassDefaultObject(FObjectSnapshotData& InObjectData, FWorldSnapshotData& InSharedData, UObject* InCDO)
 {
diff --git a/Engine/Plugins/VirtualProduction/LevelSnapshots/Source/LevelSnapshots/Private/Archive/LoadSnapshotObjectArchive.cpp b/Engine/Plugins/VirtualProduction/LevelSnapshots/Source/LevelSnapshots/Private/Archive/LoadSnapshotObjectArchive.cpp
--- a/Engine/Plugins/VirtualProduction/LevelSnapshots/Source/LevelSnapshots/Private/Archive/LoadSnapshotObjectArchive.cpp
+++ b/Engine/Plugins/VirtualProduction/LevelSnapshots/Source/LevelSnapshots/Private/Archive/LoadSnapshotObjectArchive.cpp
@@ -2,10 +2,10 @@
 
 #include "Archive/LoadSnapshotObjectArchive.h"
 
+#include "Data/Util/ObjectDependencyUtil.h"
 #include "LevelSnapshotsLog.h"
 #include "LevelSnapshotsModule.h"
 #include "WorldSnapshotData.h"
-#include "Data/Util/WorldData/SnapshotObjectUtil.h"
 
 #include "Serialization/ObjectWriter.h"
 #include "Internationalization/TextNamespaceUtil.h"
diff --git a/Engine/Plugins/VirtualProduction/LevelSnapshots/Source/LevelSnapshots/Private/Archive/LoadSnapshotObjectArchive.h b/Engine/Plugins/VirtualProduction/LevelSnapshots/Source/LevelSnapshots/Private/Archive/LoadSnapshotObjectArchive.h
--- a/Engine/Plugins/VirtualProduction/LevelSnapshots/Source/LevelSnapshots/Private/Archive/LoadSnapshotObjectArchive.h
+++ b/Engine/Plugins/VirtualProduction/LevelSnapshots/Source/LevelSnapshots/Private/Archive/LoadSnapshotObjectArchive.h
@@ -2,9 +2,8 @@
 
 #pragma once
 
-#include "CoreMinimal.h"
 #include "SnapshotArchive.h"
-#include "Misc/ObjectDependencyCallback.h"
+#include "Util/ObjectDependencyCallback.h"
 
 struct FSnapshotDataCache;
 
diff --git a/Engine/Plugins/VirtualProduction/LevelSnapshots/Source/LevelSnapshots/Private/Archive/SnapshotArchive.cpp b/Engine/Plugins/VirtualProduction/LevelSnapshots/Source/LevelSnapshots/Private/Archive/SnapshotArchive.cpp
--- a/Engine/Plugins/VirtualProduction/LevelSnapshots/Source/LevelSnapshots/Private/Archive/SnapshotArchive.cpp
+++ b/Engine/Plugins/VirtualProduction/LevelSnapshots/Source/LevelSnapshots/Private/Archive/SnapshotArchive.cpp
@@ -2,11 +2,11 @@
 
 #include "SnapshotArchive.h"
 
+#include "Data/Util/ObjectDependencyUtil.h"
 #include "LevelSnapshotsModule.h"
 #include "ObjectSnapshotData.h"
 #include "SnapshotRestorability.h"
 #include "SnapshotVersion.h"
-#include "Util/WorldData/SnapshotObjectUtil.h"
 #include "WorldSnapshotData.h"
 #if UE_BUILD_DEBUG
 #include "SnapshotConsoleVariables.h"
diff --git a/Engine/Plugins/VirtualProduction/LevelSnapshots/Source/LevelSnapshots/Private/Data/CustomSerialization/CustomObjectSerializationWrapper.cpp b/Engine/Plugins/VirtualProduction/LevelSnapshots/Source/LevelSnapshots/Private/Data/CustomSerialization/CustomObjectSerializationWrapper.cpp
--- a/Engine/Plugins/VirtualProduction/LevelSnapshots/Source/LevelSnapshots/Private/Data/CustomSerialization/CustomObjectSerializationWrapper.cpp
+++ b/Engine/Plugins/VirtualProduction/LevelSnapshots/Source/LevelSnapshots/Private/Data/CustomSerialization/CustomObjectSerializationWrapper.cpp
@@ -5,14 +5,13 @@
 #include "Archive/ApplySnapshotToEditorArchive.h"
 #include "Archive/LoadSnapshotObjectArchive.h"
 #include "CustomSerialization/CustomSerializationDataManager.h"
-#include "Data/Util/SnapshotUtil.h"
-#include "Data/Util/WorldData/SnapshotObjectUtil.h"
-#include "Data/Util/WorldData/ActorUtil.h"
+#include "Data/Util/ObjectDependencyUtil.h"
+#include "Filtering/PropertySelectionMap.h"
+#include "Restoration/Actor/ActorUtil.h"
 #include "Data/WorldSnapshotData.h"
 #include "Interfaces/ICustomObjectSnapshotSerializer.h"
 #include "LevelSnapshotsLog.h"
 #include "LevelSnapshotsModule.h"
-#include "Selection/PropertySelectionMap.h"
 
 #include "GameFramework/Actor.h"
 #include "Modules/ModuleManager.h"
diff --git a/Engine/Plugins/VirtualProduction/LevelSnapshots/Source/LevelSnapshots/Private/Data/CustomSerialization/CustomObjectSerializationWrapper.h b/Engine/Plugins/VirtualProduction/LevelSnapshots/Source/LevelSnapshots/Private/Data/CustomSerialization/CustomObjectSerializationWrapper.h
--- a/Engine/Plugins/VirtualProduction/LevelSnapshots/Source/LevelSnapshots/Private/Data/CustomSerialization/CustomObjectSerializationWrapper.h
+++ b/Engine/Plugins/VirtualProduction/LevelSnapshots/Source/LevelSnapshots/Private/Data/CustomSerialization/CustomObjectSerializationWrapper.h
@@ -2,8 +2,7 @@
 
 #pragma once
 
-#include "CoreMinimal.h"
-#include "Misc/ObjectDependencyCallback.h"
+#include "Util/ObjectDependencyCallback.h"
 #include "Templates/Function.h"
 
 struct FSnapshotDataCache;
diff --git a/Engine/Plugins/VirtualProduction/LevelSnapshots/Source/LevelSnapshots/Private/Data/CustomSerialization/CustomSerializationDataManager.cpp b/Engine/Plugins/VirtualProduction/LevelSnapshots/Source/LevelSnapshots/Private/Data/CustomSerialization/CustomSerializationDataManager.cpp
--- a/Engine/Plugins/VirtualProduction/LevelSnapshots/Source/LevelSnapshots/Private/Data/CustomSerialization/CustomSerializationDataManager.cpp
+++ b/Engine/Plugins/VirtualProduction/LevelSnapshots/Source/LevelSnapshots/Private/Data/CustomSerialization/CustomSerializationDataManager.cpp
@@ -3,16 +3,16 @@
 #include "CustomSerializationDataManager.h"
 
 #include "CustomObjectSerializationWrapper.h"
+#include "CustomSerialization/LoadCustomObjectProxyArchive.h"
+#include "CustomSerialization/SaveCustomObjectProxyArchive.h"
+#include "Data/Util/ObjectDependencyUtil.h"
 #include "LevelSnapshotsLog.h"
 #include "SnapshotCustomVersion.h"
 #include "TakeWorldObjectSnapshotArchive.h"
-#include "CustomSerialization/LoadCustomObjectProxyArchive.h"
-#include "CustomSerialization/SaveCustomObjectProxyArchive.h"
 
 #include "Serialization/MemoryWriter.h"
 #include "Serialization/MemoryReader.h"
 #include "Serialization/ObjectAndNameAsStringProxyArchive.h"
-#include "Util/WorldData/SnapshotObjectUtil.h"
 
 namespace UE::LevelSnapshots::Private::Internal
 {
diff --git a/Engine/Plugins/VirtualProduction/LevelSnapshots/Source/LevelSnapshots/Private/Data/CustomSerialization/CustomSerializationDataManager.h b/Engine/Plugins/VirtualProduction/LevelSnapshots/Source/LevelSnapshots/Private/Data/CustomSerialization/CustomSerializationDataManager.h
--- a/Engine/Plugins/VirtualProduction/LevelSnapshots/Source/LevelSnapshots/Private/Data/CustomSerialization/CustomSerializationDataManager.h
+++ b/Engine/Plugins/VirtualProduction/LevelSnapshots/Source/LevelSnapshots/Private/Data/CustomSerialization/CustomSerializationDataManager.h
@@ -2,13 +2,10 @@
 
 #pragma once
 
-#include "CoreMinimal.h"
-
 #include "CustomSerializationData.h"
 #include "WorldSnapshotData.h"
 #include "Params/ObjectSnapshotSerializationData.h"
 
-
 namespace UE::LevelSnapshots::Private
 {
 	DECLARE_DELEGATE_RetVal(const FCustomSerializationData*, FCustomSerializationDataGetter_ReadOnly);
diff --git a/Engine/Plugins/VirtualProduction/LevelSnapshots/Source/LevelSnapshots/Private/Data/LevelSnapshot.cpp b/Engine/Plugins/VirtualProduction/LevelSnapshots/Source/LevelSnapshots/Private/Data/LevelSnapshot.cpp
--- a/Engine/Plugins/VirtualProduction/LevelSnapshots/Source/LevelSnapshots/Private/Data/LevelSnapshot.cpp
+++ b/Engine/Plugins/VirtualProduction/LevelSnapshots/Source/LevelSnapshots/Private/Data/LevelSnapshot.cpp
@@ -2,11 +2,11 @@
 
 #include "Data/LevelSnapshot.h"
 
-#include "Data/Util/ActorHashUtil.h"
-#include "Data/Util/EquivalenceUtil.h"
-#include "Data/Util/SnapshotUtil.h"
-#include "Data/Util/WorldData/ActorUtil.h"
-#include "Data/Util/WorldData/WorldDataUtil.h"
+#include "Data/Hashing/ActorHashUtil.h"
+#include "Data/Util/ClassDataUtil.h"
+#include "Data/Util/WorldDataUtil.h"
+#include "Filtering/Diffing/EquivalenceUtil.h"
+#include "Restoration/Actor/ActorUtil.h"
 #include "LevelSnapshotsSettings.h"
 #include "LevelSnapshotsLog.h"
 #include "LevelSnapshotsModule.h"
@@ -14,17 +14,15 @@
 #include "SnapshotCustomVersion.h"
 #include "SnapshotConsoleVariables.h"
 #include "Util/SortedScopedLog.h"
+#include "Util/SoftObjectPathUtil.h"
 
 #include "Algo/Accumulate.h"
 #include "Engine/Engine.h"
 #include "EngineUtils.h"
-#include "Algo/IndexOf.h"
 #include "HAL/IConsoleManager.h"
 #include "GameFramework/Actor.h"
 #include "Misc/ScopeExit.h"
-#include "Stats/StatsMisc.h"
 #include "UObject/Package.h"
-#include "Util/WorldData/ClassDataUtil.h"
 #if WITH_EDITOR && !(UE_BUILD_SHIPPING || UE_BUILD_TEST)
 #include "Logging/MessageLog.h"
 #endif
diff --git a/Engine/Plugins/VirtualProduction/LevelSnapshots/Source/LevelSnapshots/Private/Data/Util/WorldDataUtil.h b/Engine/Plugins/VirtualProduction/LevelSnapshots/Source/LevelSnapshots/Private/Data/Util/WorldDataUtil.h
new file mode 100644
--- /dev/null
+++ b/Engine/Plugins/VirtualProduction/LevelSnapshots/Source/LevelSnapshots/Private/Data/Util/WorldDataUtil.h
@@ -0,0 +1,25 @@
+﻿// Copyright Epic Games, Inc. All Rights Reserved.
+
+#pragma once
+
+#include "CoreMinimal.h"
+
+class AActor;
+class UActorComponent;
+class UWorld;
+
+struct FPropertySelectionMap;
+struct FSnapshotDataCache;
+struct FWorldSnapshotData;
+
+namespace UE::LevelSnapshots::Private
+{
+	/** Captures the state of the world */
+	FWorldSnapshotData SnapshotWorld(UWorld* World);
+	
+	/* Applies the saved properties to WorldActor */
+	void ApplyToWorld(FWorldSnapshotData& WorldData, FSnapshotDataCache& Cache, UWorld* WorldToApplyTo, UPackage* LocalisationSnapshotPackage, const FPropertySelectionMap& PropertiesToSerialize);
+
+	/** Whether the component was saved in the world data */
+	bool HasSavedComponentData(const FWorldSnapshotData& WorldData, const FSoftObjectPath& WorldActorPath, const UActorComponent* EditorOrSnapshotComponent);
+}
diff --git a/Engine/Plugins/VirtualProduction/LevelSnapshots/Source/LevelSnapshots/Private/Data/WorldSnapshotData.cpp b/Engine/Plugins/VirtualProduction/LevelSnapshots/Source/LevelSnapshots/Private/Data/WorldSnapshotData.cpp
--- a/Engine/Plugins/VirtualProduction/LevelSnapshots/Source/LevelSnapshots/Private/Data/WorldSnapshotData.cpp
+++ b/Engine/Plugins/VirtualProduction/LevelSnapshots/Source/LevelSnapshots/Private/Data/WorldSnapshotData.cpp
@@ -2,15 +2,14 @@
 
 #include "Data/WorldSnapshotData.h"
 
+#include "Data/Compression/CompressionUtil.h"
+#include "Data/Util/ClassDataUtil.h"
 #include "LevelSnapshotsModule.h"
 #include "SnapshotCustomVersion.h"
-#include "Util/Property/PropertyIterator.h"
+#include "Util/PropertyIterator.h"
 
 #include "Serialization/BufferArchive.h"
 #include "UObject/UnrealType.h"
-#include "Util/SnapshotUtil.h"
-#include "Util/WorldData/ClassDataUtil.h"
-#include "Util/WorldData/CompressionUtil.h"
 
 #define LOCTEXT_NAMESPACE "LevelSnapshotsEditor"
 
diff --git a/Engine/Plugins/VirtualProduction/LevelSnapshots/Source/LevelSnapshots/Private/Filtering/Diffing/ApplySnapshotFilter.h b/Engine/Plugins/VirtualProduction/LevelSnapshots/Source/LevelSnapshots/Private/Filtering/Diffing/ApplySnapshotFilter.h
new file mode 100644
--- /dev/null
+++ b/Engine/Plugins/VirtualProduction/LevelSnapshots/Source/LevelSnapshots/Private/Filtering/Diffing/ApplySnapshotFilter.h
@@ -0,0 +1,118 @@
+﻿// Copyright Epic Games, Inc. All Rights Reserved.
+
+#pragma once
+
+#include "Filtering/PropertySelectionMap.h"
+
+class AActor;
+class UActorComponent;
+class ULevelSnapshot;
+class ULevelSnapshotFilter;
+
+namespace UE::LevelSnapshots::Private
+{
+	class FApplySnapshotFilter
+	{
+	public:
+
+		static FApplySnapshotFilter Make(ULevelSnapshot* Snapshot, AActor* DeserializedSnapshotActor, AActor* WorldActor, const ULevelSnapshotFilter* Filter);
+		
+		FApplySnapshotFilter& AllowUnchangedProperties(bool bNewValue)
+		{
+			bAllowUnchangedProperties = bNewValue;
+			return *this;
+		}
+		FApplySnapshotFilter& AllowNonEditableProperties(bool bNewValue)
+		{
+			bAllowNonEditableProperties = bNewValue;
+			return *this;
+		}
+		
+		void ApplyFilterToFindSelectedProperties(FPropertySelectionMap& MapToAddTo);
+
+	private:
+
+		enum class EFilterObjectPropertiesResult
+		{
+			HasCustomSubobjects,
+			HasOnlyNormalProperties
+		};
+		
+		struct FPropertyContainerContext
+		{
+			FPropertySelectionMap& MapToAddTo;
+			FPropertySelection& SelectionToAddTo;
+		
+			UStruct* ContainerClass;
+			void* SnapshotContainer;
+			void* WorldContainer;
+			UObject* AnalysedSnapshotObject;
+			UObject* AnalysedWorldObject;
+				
+			/* Information passed to blueprints. Property name is appended to this.
+			 * Example: [FooComponent] [BarStructPropertyName]...
+			 */
+			TArray<FString> AuthoredPathInformation;
+
+			/* Keeps track of the structs leading to this container */
+			FLevelSnapshotPropertyChain PropertyChain;
+			/* Class that PropertyChain begins from. */
+			UClass* RootClass;
+
+			FPropertyContainerContext(FPropertySelectionMap& MapToAddTo, FPropertySelection& SelectionToAddTo, UStruct* ContainerClass, void* SnapshotContainer, void* WorldContainer, UObject* AnalysedSnapshotObject, UObject* AnalysedWorldObject, const TArray<FString>& AuthoredPathInformation, const FLevelSnapshotPropertyChain& PropertyChain, UClass* RootClass);
+		};
+
+		enum class EPropertySearchResult
+		{
+			FoundProperties,
+			NoPropertiesFound
+		};
+		
+		FApplySnapshotFilter(ULevelSnapshot* Snapshot, AActor* DeserializedSnapshotActor, AActor* WorldActor, const ULevelSnapshotFilter* Filter);
+		bool EnsureParametersAreValid() const;
+		
+		void AnalyseComponentProperties(FPropertySelectionMap& MapToAddTo);
+
+		void FilterActorPair(FPropertySelectionMap& MapToAddTo);
+		EPropertySearchResult FilterSubobjectPair(FPropertySelectionMap& MapToAddTo, UObject* SnapshotSubobject, UObject* WorldSubobject);
+		EPropertySearchResult FilterStructPair(FPropertyContainerContext& Parent, FStructProperty* StructProperty);
+		EFilterObjectPropertiesResult FindAndFilterCustomSubobjectPairs(FPropertySelectionMap& MapToAddTo, UObject* SnapshotOwner, UObject* WorldOwner);
+
+		void AnalyseRootProperties(FPropertyContainerContext& ContainerContext, UObject* SnapshotObject, UObject* WorldObject);
+		void ExtendAnalysedProperties(FPropertyContainerContext& ContainerContext, UObject* SnapshotObject, UObject* WorldObject);
+		TOptional<EPropertySearchResult> HandlePossibleStructProperties(FPropertyContainerContext& ContainerContext, FProperty* PropertyToHandle);
+		EPropertySearchResult AnalyseStructProperties(FPropertyContainerContext& ContainerContext);
+
+		
+		TOptional<EPropertySearchResult> HandlePossibleSubobjectProperties(FPropertyContainerContext& ContainerContext, FProperty* PropertyToHandle);
+		
+		TOptional<EPropertySearchResult> AnalysePossibleArraySubobjectProperties(FPropertyContainerContext& ContainerContext, FArrayProperty* PropertyToHandle);
+		TOptional<EPropertySearchResult> AnalysePossibleSetSubobjectProperties(FPropertyContainerContext& ContainerContext, FSetProperty* PropertyToHandle);
+		TOptional<EPropertySearchResult> AnalysePossibleMapSubobjectProperties(FPropertyContainerContext& ContainerContext, FMapProperty* PropertyToHandle);
+		
+		template<typename TCollectionData>
+		TOptional<EPropertySearchResult> AnalysePossibleSubobjectsInCollection(FPropertyContainerContext& ContainerContext, FObjectPropertyBase* ObjectProperty, TCollectionData& Detail);
+		
+		TOptional<EPropertySearchResult> AnalysePossibleSubobjectProperties(FPropertyContainerContext& ContainerContext, FObjectPropertyBase* PropertyToHandle, void* SnapshotValuePtr, void* WorldValuePtr);
+
+		
+
+		EPropertySearchResult AnalyseProperty(FPropertyContainerContext& ContainerContext, FProperty* PropertyInCommon, bool bSkipEqualityTest = false);
+		EPropertySearchResult TrackChangedProperties(FPropertyContainerContext& ContainerContext, FProperty* PropertyInCommon);
+
+		
+		ULevelSnapshot* Snapshot;
+		AActor* DeserializedSnapshotActor;
+		AActor* WorldActor;
+		const ULevelSnapshotFilter* Filter;
+		/** Every time an object is analysed, it is added here, so we do not analyse more than once. */
+		TMap<UObject*, EPropertySearchResult> AnalysedSnapshotObjects;
+
+		/* Do we allow properties that do not show up in the details panel? */
+		bool bAllowNonEditableProperties = false;
+
+		/* Do we allow adding properties that did not change to the selection map? */
+		bool bAllowUnchangedProperties = false;
+	};
+}
+
diff --git a/Engine/Plugins/VirtualProduction/LevelSnapshots/Source/LevelSnapshots/Private/LevelSnapshotsFilteringLibrary.cpp b/Engine/Plugins/VirtualProduction/LevelSnapshots/Source/LevelSnapshots/Private/LevelSnapshotsFilteringLibrary.cpp
--- a/Engine/Plugins/VirtualProduction/LevelSnapshots/Source/LevelSnapshots/Private/LevelSnapshotsFilteringLibrary.cpp
+++ b/Engine/Plugins/VirtualProduction/LevelSnapshots/Source/LevelSnapshots/Private/LevelSnapshotsFilteringLibrary.cpp
@@ -2,11 +2,11 @@
 
 #include "LevelSnapshotsFilteringLibrary.h"
 
-#include "Selection/ApplySnapshotFilter.h"
 #include "ConstantFilter.h"
+#include "Filtering/Diffing/ApplySnapshotFilter.h"
+#include "Filtering/SnapshotSelectionBuilder.h"
 #include "LevelSnapshot.h"
 #include "LevelSnapshotFilters.h"
-#include "Selection/SnapshotSelectionBuilder.h"
 
 #include "EngineUtils.h"
 #include "GameFramework/Actor.h"
diff --git a/Engine/Plugins/VirtualProduction/LevelSnapshots/Source/LevelSnapshots/Private/LevelSnapshotsFunctionLibrary.cpp b/Engine/Plugins/VirtualProduction/LevelSnapshots/Source/LevelSnapshots/Private/LevelSnapshotsFunctionLibrary.cpp
--- a/Engine/Plugins/VirtualProduction/LevelSnapshots/Source/LevelSnapshots/Private/LevelSnapshotsFunctionLibrary.cpp
+++ b/Engine/Plugins/VirtualProduction/LevelSnapshots/Source/LevelSnapshots/Private/LevelSnapshotsFunctionLibrary.cpp
@@ -2,7 +2,7 @@
 
 #include "LevelSnapshotsFunctionLibrary.h"
 
-#include "Selection/ApplySnapshotFilter.h"
+#include "Filtering/Diffing/ApplySnapshotFilter.h"
 #include "LevelSnapshot.h"
 #include "LevelSnapshotFilters.h"
 
diff --git a/Engine/Plugins/VirtualProduction/LevelSnapshots/Source/LevelSnapshots/Private/LevelSnapshotsModule.cpp b/Engine/Plugins/VirtualProduction/LevelSnapshots/Source/LevelSnapshots/Private/LevelSnapshotsModule.cpp
--- a/Engine/Plugins/VirtualProduction/LevelSnapshots/Source/LevelSnapshots/Private/LevelSnapshotsModule.cpp
+++ b/Engine/Plugins/VirtualProduction/LevelSnapshots/Source/LevelSnapshots/Private/LevelSnapshotsModule.cpp
@@ -6,13 +6,13 @@
 #include "LevelSnapshotsSettings.h"
 #include "LevelSnapshotsLog.h"
 #include "Params/PropertyComparisonParams.h"
+#include "SpecialSupport/EngineTypesRestorationFence.h"
 
 #include "Algo/AllOf.h"
 #include "Components/ActorComponent.h"
 #include "EngineUtils.h"
 #include "GameFramework/Actor.h"
 #include "Modules/ModuleManager.h"
-#include "Restorability/EngineTypesRestorationFence.h"
 #if WITH_EDITOR
 #include "ISettingsModule.h"
 #endif
diff --git a/Engine/Plugins/VirtualProduction/LevelSnapshots/Source/LevelSnapshots/Private/Settings/LevelSnapshotsSettings.cpp b/Engine/Plugins/VirtualProduction/LevelSnapshots/Source/LevelSnapshots/Private/Settings/LevelSnapshotsSettings.cpp
--- a/Engine/Plugins/VirtualProduction/LevelSnapshots/Source/LevelSnapshots/Private/Settings/LevelSnapshotsSettings.cpp
+++ b/Engine/Plugins/VirtualProduction/LevelSnapshots/Source/LevelSnapshots/Private/Settings/LevelSnapshotsSettings.cpp
@@ -2,7 +2,7 @@
 
 #include "Settings/LevelSnapshotsSettings.h"
 
-#include "Data/Util/ActorHashUtil.h"
+#include "Data/Hashing/ActorHashUtil.h"
 #include "PropertyInfoHelpers.h"
 
 ULevelSnapshotsSettings* ULevelSnapshotsSettings::Get()
diff --git a/Engine/Plugins/VirtualProduction/LevelSnapshots/Source/LevelSnapshots/Private/Tests/ApplyToWorldTests.cpp b/Engine/Plugins/VirtualProduction/LevelSnapshots/Source/LevelSnapshots/Private/Tests/ApplyToWorldTests.cpp
--- a/Engine/Plugins/VirtualProduction/LevelSnapshots/Source/LevelSnapshots/Private/Tests/ApplyToWorldTests.cpp
+++ b/Engine/Plugins/VirtualProduction/LevelSnapshots/Source/LevelSnapshots/Private/Tests/ApplyToWorldTests.cpp
@@ -1,8 +1,8 @@
 ﻿// Copyright Epic Games, Inc. All Rights Reserved.
 
-#include "Selection/PropertySelectionMap.h"
+#include "Filtering/PropertySelectionMap.h"
+#include "Filtering/Diffing/EquivalenceUtil.h"
 #include "Util/SnapshotTestRunner.h"
-#include "Util/EquivalenceUtil.h"
 #include "Types/SnapshotTestActor.h"
 
 #include "Components/PointLightComponent.h"
diff --git a/Engine/Plugins/VirtualProduction/LevelSnapshots/Source/LevelSnapshots/Private/Tests/ComponentTests.cpp b/Engine/Plugins/VirtualProduction/LevelSnapshots/Source/LevelSnapshots/Private/Tests/ComponentTests.cpp
--- a/Engine/Plugins/VirtualProduction/LevelSnapshots/Source/LevelSnapshots/Private/Tests/ComponentTests.cpp
+++ b/Engine/Plugins/VirtualProduction/LevelSnapshots/Source/LevelSnapshots/Private/Tests/ComponentTests.cpp
@@ -1,7 +1,9 @@
 ﻿// Copyright Epic Games, Inc. All Rights Reserved.
 
-#include "Selection/PropertySelection.h"
-#include "Selection/PropertySelectionMap.h"
+#include "Filtering/PropertySelection.h"
+
+#include "Filtering/Diffing/EquivalenceUtil.h"
+#include "Filtering/PropertySelectionMap.h"
 #include "Util/SnapshotTestRunner.h"
 #include "Types/SnapshotTestActor.h"
 
@@ -13,7 +15,6 @@
 #include "Engine/World.h"
 #include "GameFramework/Actor.h"
 #include "Misc/AutomationTest.h"
-#include "Util/EquivalenceUtil.h"
 
 namespace UE::LevelSnapshots::Private::Tests
 {
diff --git a/Engine/Plugins/VirtualProduction/LevelSnapshots/Source/LevelSnapshots/Private/Tests/CustomObjectSerializationTests.cpp b/Engine/Plugins/VirtualProduction/LevelSnapshots/Source/LevelSnapshots/Private/Tests/CustomObjectSerializationTests.cpp
--- a/Engine/Plugins/VirtualProduction/LevelSnapshots/Source/LevelSnapshots/Private/Tests/CustomObjectSerializationTests.cpp
+++ b/Engine/Plugins/VirtualProduction/LevelSnapshots/Source/LevelSnapshots/Private/Tests/CustomObjectSerializationTests.cpp
@@ -1,16 +1,16 @@
 ﻿// Copyright Epic Games, Inc. All Rights Reserved.
 
-#include "EngineUtils.h"
-#include "LevelSnapshotsModule.h"
+#include "Filtering/CustomSubobjectRestorationInfo.h"
+#include "Filtering/PropertySelectionMap.h"
 #include "Interfaces/ICustomObjectSnapshotSerializer.h"
+#include "LevelSnapshotsModule.h"
 #include "Params/ObjectSnapshotSerializationData.h"
-#include "Selection/CustomSubobjectRestorationInfo.h"
-#include "Selection/PropertySelectionMap.h"
 #include "Types/SnapshotTestActor.h"
 #include "Util/CustomSubobjectTestUtil.h"
 #include "Util/SnapshotTestRunner.h"
 
 #include "Engine/StaticMeshActor.h"
+#include "EngineUtils.h"
 #include "Misc/AutomationTest.h"
 #include "Misc/ScopeExit.h"
 
diff --git a/Engine/Plugins/VirtualProduction/LevelSnapshots/Source/LevelSnapshots/Private/Tests/RegressionTests.cpp b/Engine/Plugins/VirtualProduction/LevelSnapshots/Source/LevelSnapshots/Private/Tests/RegressionTests.cpp
--- a/Engine/Plugins/VirtualProduction/LevelSnapshots/Source/LevelSnapshots/Private/Tests/RegressionTests.cpp
+++ b/Engine/Plugins/VirtualProduction/LevelSnapshots/Source/LevelSnapshots/Private/Tests/RegressionTests.cpp
@@ -1,8 +1,6 @@
 ﻿// Copyright Epic Games, Inc. All Rights Reserved.
 
-#include "CoreMinimal.h"
-
-#include "Selection/PropertySelectionMap.h"
+#include "Filtering/PropertySelectionMap.h"
 #include "Util/SnapshotTestRunner.h"
 #include "Types/SnapshotTestActor.h"
 #include "Types/ActorWithReferencesInCDO.h"
diff --git a/Engine/Plugins/VirtualProduction/LevelSnapshots/Source/LevelSnapshots/Public/Data/ActorSnapshotData.h b/Engine/Plugins/VirtualProduction/LevelSnapshots/Source/LevelSnapshots/Public/Data/ActorSnapshotData.h
--- a/Engine/Plugins/VirtualProduction/LevelSnapshots/Source/LevelSnapshots/Public/Data/ActorSnapshotData.h
+++ b/Engine/Plugins/VirtualProduction/LevelSnapshots/Source/LevelSnapshots/Public/Data/ActorSnapshotData.h
@@ -2,10 +2,9 @@
 
 #pragma once
 
-#include "CoreMinimal.h"
-#include "ActorSnapshotHash.h"
 #include "ComponentSnapshotData.h"
 #include "CustomSerializationData.h"
+#include "Hashing/ActorSnapshotHash.h"
 #include "ObjectSnapshotData.h"
 #include "SnapshotUtilTypes.h"
 #include "ActorSnapshotData.generated.h"
diff --git a/Engine/Plugins/VirtualProduction/LevelSnapshots/Source/LevelSnapshots/Public/Data/ComponentSnapshotData.h b/Engine/Plugins/VirtualProduction/LevelSnapshots/Source/LevelSnapshots/Public/Data/ComponentSnapshotData.h
--- a/Engine/Plugins/VirtualProduction/LevelSnapshots/Source/LevelSnapshots/Public/Data/ComponentSnapshotData.h
+++ b/Engine/Plugins/VirtualProduction/LevelSnapshots/Source/LevelSnapshots/Public/Data/ComponentSnapshotData.h
@@ -2,16 +2,9 @@
 
 #pragma once
 
-#include "CoreMinimal.h"
 #include "ComponentInstanceDataCache.h"
-#include "ObjectSnapshotData.h"
-#include "Misc/ObjectDependencyCallback.h"
 #include "ComponentSnapshotData.generated.h"
 
-class UPackage;
-struct FWorldSnapshotData;
-struct FPropertySelection;
-
 USTRUCT()
 struct LEVELSNAPSHOTS_API FComponentSnapshotData
 {
diff --git a/Engine/Plugins/VirtualProduction/LevelSnapshots/Source/LevelSnapshots/Public/Data/SnapshotUtilTypes.h b/Engine/Plugins/VirtualProduction/LevelSnapshots/Source/LevelSnapshots/Public/Data/SnapshotUtilTypes.h
--- a/Engine/Plugins/VirtualProduction/LevelSnapshots/Source/LevelSnapshots/Public/Data/SnapshotUtilTypes.h
+++ b/Engine/Plugins/VirtualProduction/LevelSnapshots/Source/LevelSnapshots/Public/Data/SnapshotUtilTypes.h
@@ -6,9 +6,6 @@
 
 namespace UE::LevelSnapshots
 {
-	// TODO:
-	/*using FSerializedObjectReferenceIndex = uint32;
-	using FSerializedNameIndex = uint32;*/
 	/** An index to FWorldSnapshotData::ClassData. */
 	using FClassDataIndex = uint32;
 }
\ No newline at end of file
diff --git a/Engine/Plugins/VirtualProduction/LevelSnapshots/Source/LevelSnapshots/Public/LevelSnapshotsFunctionLibrary.h b/Engine/Plugins/VirtualProduction/LevelSnapshots/Source/LevelSnapshots/Public/LevelSnapshotsFunctionLibrary.h
--- a/Engine/Plugins/VirtualProduction/LevelSnapshots/Source/LevelSnapshots/Public/LevelSnapshotsFunctionLibrary.h
+++ b/Engine/Plugins/VirtualProduction/LevelSnapshots/Source/LevelSnapshots/Public/LevelSnapshotsFunctionLibrary.h
@@ -2,8 +2,7 @@
 
 #pragma once
 
-#include "CoreMinimal.h"
-#include "Selection/PropertySelectionMap.h"
+#include "Filtering/PropertySelectionMap.h"
 #include "Kismet/BlueprintFunctionLibrary.h"
 #include "LevelSnapshotsFunctionLibrary.generated.h"
 
diff --git a/Engine/Plugins/VirtualProduction/LevelSnapshots/Source/LevelSnapshots/Public/Restorability/Interfaces/ISnapshotFilterExtender.h b/Engine/Plugins/VirtualProduction/LevelSnapshots/Source/LevelSnapshots/Public/Restorability/Interfaces/ISnapshotFilterExtender.h
--- a/Engine/Plugins/VirtualProduction/LevelSnapshots/Source/LevelSnapshots/Public/Restorability/Interfaces/ISnapshotFilterExtender.h
+++ b/Engine/Plugins/VirtualProduction/LevelSnapshots/Source/LevelSnapshots/Public/Restorability/Interfaces/ISnapshotFilterExtender.h
@@ -2,8 +2,7 @@
 
 #pragma once
 
-#include "CoreMinimal.h"
-#include "Selection/PropertySelection.h"
+#include "Filtering/PropertySelection.h"
 
 class AActor;
 class UObject;
diff --git a/Engine/Plugins/VirtualProduction/LevelSnapshots/Source/LevelSnapshotsEditor/Private/Data/FilterListData.h b/Engine/Plugins/VirtualProduction/LevelSnapshots/Source/LevelSnapshotsEditor/Private/Data/FilterListData.h
--- a/Engine/Plugins/VirtualProduction/LevelSnapshots/Source/LevelSnapshotsEditor/Private/Data/FilterListData.h
+++ b/Engine/Plugins/VirtualProduction/LevelSnapshots/Source/LevelSnapshotsEditor/Private/Data/FilterListData.h
@@ -2,10 +2,8 @@
 
 #pragma once
 
-#include "CoreMinimal.h"
-
 #include "GameFramework/Actor.h"
-#include "Selection/PropertySelectionMap.h"
+#include "Filtering/PropertySelectionMap.h"
 
 class AActor;
 class ULevelSnapshot;
diff --git a/Engine/Plugins/VirtualProduction/LevelSnapshots/Source/LevelSnapshotsEditor/Private/Misc/LevelSnapshotsEditorCustomWidgetGenerator.cpp b/Engine/Plugins/VirtualProduction/LevelSnapshots/Source/LevelSnapshotsEditor/Private/Misc/LevelSnapshotsEditorCustomWidgetGenerator.cpp
--- a/Engine/Plugins/VirtualProduction/LevelSnapshots/Source/LevelSnapshotsEditor/Private/Misc/LevelSnapshotsEditorCustomWidgetGenerator.cpp
+++ b/Engine/Plugins/VirtualProduction/LevelSnapshots/Source/LevelSnapshotsEditor/Private/Misc/LevelSnapshotsEditorCustomWidgetGenerator.cpp
@@ -3,10 +3,10 @@
 #include "Misc/LevelSnapshotsEditorCustomWidgetGenerator.h"
 
 #include "LevelSnapshotsLog.h"
-#include "Selection/PropertySelection.h"
+#include "Filtering/PropertySelection.h"
+#include "Views/Results/LevelSnapshotsEditorResultsRow.h"
 
 #include "Styling/AppStyle.h"
-#include "Views/Results/LevelSnapshotsEditorResultsRow.h"
 #include "Widgets/Input/SCheckBox.h"
 #include "Widgets/Input/SComboBox.h"
 #include "Widgets/Input/SNumericEntryBox.h"
