Expose some C++ LevelSnapshots events to Blueprint scripting via UEngineSubsystem

// preflight was done but ID is not available right now

diff --git a/Engine/Plugins/VirtualProduction/LevelSnapshots/Source/FoliageSupport/Private/FoliageSupport/FoliageSupport.cpp b/Engine/Plugins/VirtualProduction/LevelSnapshots/Source/FoliageSupport/Private/FoliageSupport/FoliageSupport.cpp
--- a/Engine/Plugins/VirtualProduction/LevelSnapshots/Source/FoliageSupport/Private/FoliageSupport/FoliageSupport.cpp
+++ b/Engine/Plugins/VirtualProduction/LevelSnapshots/Source/FoliageSupport/Private/FoliageSupport/FoliageSupport.cpp
@@ -190,7 +190,7 @@ namespace UE::LevelSnapshots::Foliage::Private
 		// Rest is done in PostApplySnapshotToActor (need access to the property selection map)
 	}
 
-	void FFoliageSupport::PostApplySnapshot(const FApplySnapshotParams& Params)
+	void FFoliageSupport::PostApplySnapshot(const FPostApplySnapshotParams& Params)
 	{
 		// Without this step all foliage will teleport visually upon being selected by the user; its cache needs to be updated:
 		// 1. Add 1 instance 2. Take snapshot 3. Paint many instances 4. Restore 5. Lasso select all > The instance jumps to a random place until you paint
diff --git a/Engine/Plugins/VirtualProduction/LevelSnapshots/Source/FoliageSupport/Private/FoliageSupport/FoliageSupport.h b/Engine/Plugins/VirtualProduction/LevelSnapshots/Source/FoliageSupport/Private/FoliageSupport/FoliageSupport.h
--- a/Engine/Plugins/VirtualProduction/LevelSnapshots/Source/FoliageSupport/Private/FoliageSupport/FoliageSupport.h
+++ b/Engine/Plugins/VirtualProduction/LevelSnapshots/Source/FoliageSupport/Private/FoliageSupport/FoliageSupport.h
@@ -46,7 +46,7 @@ namespace UE::LevelSnapshots::Foliage::Private
 		//~ End ICustomObjectSnapshotSerializer Interface
 		
 		//~ Begin IRestorationListener Interface
-		virtual void PostApplySnapshot(const FApplySnapshotParams& Params) override;
+		virtual void PostApplySnapshot(const FPostApplySnapshotParams& Params) override;
 		virtual void PostApplySnapshotToActor(const FApplySnapshotToActorParams& Params) override;
 		virtual void PreApplySnapshotProperties(const FApplySnapshotPropertiesParams& Params) override;
 		virtual void PostApplySnapshotProperties(const FApplySnapshotPropertiesParams& Params) override;
diff --git a/Engine/Plugins/VirtualProduction/LevelSnapshots/Source/LevelSnapshots/Private/Data/LevelSnapshot.cpp b/Engine/Plugins/VirtualProduction/LevelSnapshots/Source/LevelSnapshots/Private/Data/LevelSnapshot.cpp
--- a/Engine/Plugins/VirtualProduction/LevelSnapshots/Source/LevelSnapshots/Private/Data/LevelSnapshot.cpp
+++ b/Engine/Plugins/VirtualProduction/LevelSnapshots/Source/LevelSnapshots/Private/Data/LevelSnapshot.cpp
@@ -43,9 +43,19 @@ void ULevelSnapshot::ApplySnapshotToWorld(UWorld* TargetWorld, const FPropertySe
 	};
 	
 	EnsureWorldInitialised();
-	
+
+	using namespace UE::LevelSnapshots;
+	using namespace UE::LevelSnapshots::Private;
 	OnPreApplySnapshotDelegate.Broadcast();
-	UE::LevelSnapshots::Private::ApplyToWorld(SerializedData, Cache, TargetWorld, GetPackage(), SelectionSet);
+	ApplyToWorld(
+		SerializedData,
+		Cache,
+		TargetWorld,
+		GetPackage(),
+		SelectionSet,
+		[this, &SelectionSet](){ FLevelSnapshotsModule::GetInternalModuleInstance().OnPreApplySnapshot({FApplySnapshotSharedParams{ *this, SelectionSet }}); },
+		[this, &SelectionSet](){ FLevelSnapshotsModule::GetInternalModuleInstance().OnPostApplySnapshot({ FApplySnapshotSharedParams{ *this, SelectionSet }}); }
+		);
 	OnPostApplySnapshotDelegate.Broadcast();
 }
 
diff --git a/Engine/Plugins/VirtualProduction/LevelSnapshots/Source/LevelSnapshots/Private/Data/Util/WorldDataUtil.cpp b/Engine/Plugins/VirtualProduction/LevelSnapshots/Source/LevelSnapshots/Private/Data/Util/WorldDataUtil.cpp
--- a/Engine/Plugins/VirtualProduction/LevelSnapshots/Source/LevelSnapshots/Private/Data/Util/WorldDataUtil.cpp
+++ b/Engine/Plugins/VirtualProduction/LevelSnapshots/Source/LevelSnapshots/Private/Data/Util/WorldDataUtil.cpp
@@ -26,6 +26,7 @@
 #include "Misc/ScopedSlowTask.h"
 #include "Templates/NonNullPointer.h"
 #include "UObject/Package.h"
+#include "UObject/SoftObjectPath.h"
 #if WITH_EDITOR
 #include "Editor.h"
 #include "Editor/EditorEngine.h"
@@ -209,7 +210,14 @@ namespace UE::LevelSnapshots::Private::Internal
 #endif
 }
 
-void UE::LevelSnapshots::Private::ApplyToWorld(FWorldSnapshotData& WorldData, FSnapshotDataCache& Cache, UWorld* WorldToApplyTo, UPackage* LocalisationSnapshotPackage, const FPropertySelectionMap& PropertiesToSerialize)
+void UE::LevelSnapshots::Private::ApplyToWorld(
+	FWorldSnapshotData& WorldData,
+	FSnapshotDataCache& Cache,
+	UWorld* WorldToApplyTo,
+	UPackage* LocalisationSnapshotPackage,
+	const FPropertySelectionMap& PropertiesToSerialize,
+	TFunctionRef<void()> OnPreApply,
+	TFunctionRef<void()> OnPostApply)
 {
 	check(WorldToApplyTo);
 	
@@ -229,11 +237,8 @@ void UE::LevelSnapshots::Private::ApplyToWorld(FWorldSnapshotData& WorldData, FS
 #endif
 
 	// It's important to call these events after the transaction has been started
-	FLevelSnapshotsModule::GetInternalModuleInstance().OnPreApplySnapshot({ PropertiesToSerialize });
-	ON_SCOPE_EXIT
-	{
-		FLevelSnapshotsModule::GetInternalModuleInstance().OnPostApplySnapshot({ PropertiesToSerialize });
-	};
+	OnPreApply();
+	ON_SCOPE_EXIT { OnPostApply(); };
 
 	// Clear editor world subobject cache from previous ApplyToWorld
 	for (auto SubobjectIt = Cache.SubobjectCache.CreateIterator(); SubobjectIt; ++SubobjectIt)
diff --git a/Engine/Plugins/VirtualProduction/LevelSnapshots/Source/LevelSnapshots/Private/Data/Util/WorldDataUtil.h b/Engine/Plugins/VirtualProduction/LevelSnapshots/Source/LevelSnapshots/Private/Data/Util/WorldDataUtil.h
--- a/Engine/Plugins/VirtualProduction/LevelSnapshots/Source/LevelSnapshots/Private/Data/Util/WorldDataUtil.h
+++ b/Engine/Plugins/VirtualProduction/LevelSnapshots/Source/LevelSnapshots/Private/Data/Util/WorldDataUtil.h
@@ -2,7 +2,7 @@
 
 #pragma once
 
-#include "CoreMinimal.h"
+#include "Templates/Function.h"
 
 class AActor;
 class UActorComponent;
@@ -10,6 +10,7 @@ class UWorld;
 
 struct FPropertySelectionMap;
 struct FSnapshotDataCache;
+struct FSoftObjectPath;
 struct FWorldSnapshotData;
 
 namespace UE::LevelSnapshots::Private
@@ -18,7 +19,15 @@ namespace UE::LevelSnapshots::Private
 	FWorldSnapshotData SnapshotWorld(UWorld* World);
 	
 	/* Applies the saved properties to WorldActor */
-	void ApplyToWorld(FWorldSnapshotData& WorldData, FSnapshotDataCache& Cache, UWorld* WorldToApplyTo, UPackage* LocalisationSnapshotPackage, const FPropertySelectionMap& PropertiesToSerialize);
+	void ApplyToWorld(
+		FWorldSnapshotData& WorldData,
+		FSnapshotDataCache& Cache,
+		UWorld* WorldToApplyTo,
+		UPackage* LocalisationSnapshotPackage,
+		const FPropertySelectionMap& PropertiesToSerialize,
+		TFunctionRef<void()> OnPreApply,
+		TFunctionRef<void()> OnPostApply
+		);
 
 	/** Whether the component was saved in the world data */
 	bool HasSavedComponentData(const FWorldSnapshotData& WorldData, const FSoftObjectPath& WorldActorPath, const UActorComponent* EditorOrSnapshotComponent);
diff --git a/Engine/Plugins/VirtualProduction/LevelSnapshots/Source/LevelSnapshots/Private/LevelSnapshotsEngineSubsystem.cpp b/Engine/Plugins/VirtualProduction/LevelSnapshots/Source/LevelSnapshots/Private/LevelSnapshotsEngineSubsystem.cpp
new file mode 100644
--- /dev/null
+++ b/Engine/Plugins/VirtualProduction/LevelSnapshots/Source/LevelSnapshots/Private/LevelSnapshotsEngineSubsystem.cpp
@@ -0,0 +1,68 @@
+// Copyright Epic Games, Inc. All Rights Reserved.
+
+#include "LevelSnapshotsEngineSubsystem.h"
+
+#include "ILevelSnapshotsModule.h"
+
+namespace UE::LevelSnapshots::Private
+{
+	class FSubsystemRestorationListenerAdapter : public IRestorationListener
+	{
+		/** Lifetime is tracked by ULevelSnapshotsEngineSubsystem::Deinitialize */
+		ULevelSnapshotsEngineSubsystem& OwningSubsystem;
+	public:
+
+		FSubsystemRestorationListenerAdapter(ULevelSnapshotsEngineSubsystem& OwningSubsystem)
+			: OwningSubsystem(OwningSubsystem)
+		{}
+		
+		virtual void PreApplySnapshot(const FPreApplySnapshotParams& Params) override
+		{
+			OwningSubsystem.OnPreApplySnapshot.Broadcast({ &Params.LevelSnapshot });
+		}
+		
+		virtual void PostApplySnapshot(const FPostApplySnapshotParams& Params) override
+		{
+			OwningSubsystem.OnPostApplySnapshot.Broadcast({ &Params.LevelSnapshot });
+		}
+	};
+}
+
+void ULevelSnapshotsEngineSubsystem::Initialize(FSubsystemCollectionBase& Collection)
+{
+	Super::Initialize(Collection);
+
+	using namespace UE::LevelSnapshots;
+	ILevelSnapshotsModule& Module = ILevelSnapshotsModule::Get();
+	Module.OnPreTakeSnapshot().AddUObject(this, &ULevelSnapshotsEngineSubsystem::HandleOnPreTakeSnapshot);
+	Module.OnPostTakeSnapshot().AddUObject(this, &ULevelSnapshotsEngineSubsystem::HandleOnPostTakeSnapshot);
+
+	RestorationListener = MakeShared<Private::FSubsystemRestorationListenerAdapter>(*this);
+	Module.RegisterRestorationListener(RestorationListener.ToSharedRef());
+}
+
+void ULevelSnapshotsEngineSubsystem::Deinitialize()
+{
+	Super::Deinitialize();
+
+	using namespace UE::LevelSnapshots;
+	if (ILevelSnapshotsModule::IsAvailable())
+	{
+		ILevelSnapshotsModule& Module = ILevelSnapshotsModule::Get();
+		Module.OnPreTakeSnapshot().RemoveAll(this);
+		Module.OnPostTakeSnapshot().RemoveAll(this);
+		Module.UnregisterRestorationListener(RestorationListener.ToSharedRef());
+	}
+
+	RestorationListener.Reset();
+}
+
+void ULevelSnapshotsEngineSubsystem::HandleOnPreTakeSnapshot(const UE::LevelSnapshots::FPreTakeSnapshotEventData& EventData)
+{
+	OnPreTakeSnapshot.Broadcast({ EventData.Snapshot });
+}
+
+void ULevelSnapshotsEngineSubsystem::HandleOnPostTakeSnapshot(const UE::LevelSnapshots::FPostTakeSnapshotEventData& EventData)
+{
+	OnPostTakeSnapshot.Broadcast({ EventData.Snapshot });
+}
diff --git a/Engine/Plugins/VirtualProduction/LevelSnapshots/Source/LevelSnapshots/Private/LevelSnapshotsModule.cpp b/Engine/Plugins/VirtualProduction/LevelSnapshots/Source/LevelSnapshots/Private/LevelSnapshotsModule.cpp
--- a/Engine/Plugins/VirtualProduction/LevelSnapshots/Source/LevelSnapshots/Private/LevelSnapshotsModule.cpp
+++ b/Engine/Plugins/VirtualProduction/LevelSnapshots/Source/LevelSnapshots/Private/LevelSnapshotsModule.cpp
@@ -409,7 +409,7 @@ void UE::LevelSnapshots::Private::FLevelSnapshotsModule::OnPostLoadSnapshotObjec
 	}
 }
 
-void UE::LevelSnapshots::Private::FLevelSnapshotsModule::OnPreApplySnapshot(const FApplySnapshotParams& Params)
+void UE::LevelSnapshots::Private::FLevelSnapshotsModule::OnPreApplySnapshot(const FPreApplySnapshotParams& Params)
 {
 	SCOPED_SNAPSHOT_CORE_TRACE(RestorationListeners);
 	
@@ -419,7 +419,7 @@ void UE::LevelSnapshots::Private::FLevelSnapshotsModule::OnPreApplySnapshot(cons
 	}
 }
 
-void UE::LevelSnapshots::Private::FLevelSnapshotsModule::OnPostApplySnapshot(const FApplySnapshotParams& Params)
+void UE::LevelSnapshots::Private::FLevelSnapshotsModule::OnPostApplySnapshot(const FPostApplySnapshotParams& Params)
 {
 	SCOPED_SNAPSHOT_CORE_TRACE(RestorationListeners);
 	
diff --git a/Engine/Plugins/VirtualProduction/LevelSnapshots/Source/LevelSnapshots/Private/LevelSnapshotsModule.h b/Engine/Plugins/VirtualProduction/LevelSnapshots/Source/LevelSnapshots/Private/LevelSnapshotsModule.h
--- a/Engine/Plugins/VirtualProduction/LevelSnapshots/Source/LevelSnapshots/Private/LevelSnapshotsModule.h
+++ b/Engine/Plugins/VirtualProduction/LevelSnapshots/Source/LevelSnapshots/Private/LevelSnapshotsModule.h
@@ -74,8 +74,8 @@ namespace UE::LevelSnapshots::Private
 
 		void OnPostLoadSnapshotObject(const FPostLoadSnapshotObjectParams& Params);
 
-		void OnPreApplySnapshot(const FApplySnapshotParams& Params);
-		void OnPostApplySnapshot(const FApplySnapshotParams& Params);
+		void OnPreApplySnapshot(const FPreApplySnapshotParams& Params);
+		void OnPostApplySnapshot(const FPostApplySnapshotParams& Params);
 		
 		void OnPreApplySnapshotProperties(const FApplySnapshotPropertiesParams& Params);
 		void OnPostApplySnapshotProperties(const FApplySnapshotPropertiesParams& Params);
diff --git a/Engine/Plugins/VirtualProduction/LevelSnapshots/Source/LevelSnapshots/Public/ILevelSnapshotsModule.h b/Engine/Plugins/VirtualProduction/LevelSnapshots/Source/LevelSnapshots/Public/ILevelSnapshotsModule.h
--- a/Engine/Plugins/VirtualProduction/LevelSnapshots/Source/LevelSnapshots/Public/ILevelSnapshotsModule.h
+++ b/Engine/Plugins/VirtualProduction/LevelSnapshots/Source/LevelSnapshots/Public/ILevelSnapshotsModule.h
@@ -36,6 +36,11 @@ namespace UE::LevelSnapshots
 		{
 			return FModuleManager::Get().GetModuleChecked<ILevelSnapshotsModule>("LevelSnapshots");
 		}
+
+		static bool IsAvailable()
+		{
+			return FModuleManager::Get().IsModuleLoaded("LevelSnapshots");
+		}
 		
 
 		DECLARE_EVENT_OneParam(ILevelSnapshotsModule, FPreTakeSnapshotEvent, const FPreTakeSnapshotEventData&);
diff --git a/Engine/Plugins/VirtualProduction/LevelSnapshots/Source/LevelSnapshots/Public/LevelSnapshotsEngineSubsystem.h b/Engine/Plugins/VirtualProduction/LevelSnapshots/Source/LevelSnapshots/Public/LevelSnapshotsEngineSubsystem.h
new file mode 100644
--- /dev/null
+++ b/Engine/Plugins/VirtualProduction/LevelSnapshots/Source/LevelSnapshots/Public/LevelSnapshotsEngineSubsystem.h
@@ -0,0 +1,79 @@
+// Copyright Epic Games, Inc. All Rights Reserved.
+
+#pragma once
+
+#include "Subsystems/EngineSubsystem.h"
+#include "Templates/SharedPointer.h"
+#include "LevelSnapshotsEngineSubsystem.generated.h"
+
+class ULevelSnapshot;
+namespace UE::LevelSnapshots
+{
+	class IRestorationListener;
+	struct FPreTakeSnapshotEventData;
+	struct FPostTakeSnapshotEventData;
+}
+
+// Event data is extracted to structs so event delegate signature does not need to be modified in the future.
+// Please follow this pattern when you add new events.
+
+USTRUCT(BlueprintType, DisplayName = "LevelSnapshotEvent")
+struct LEVELSNAPSHOTS_API FLevelSnapshotEvent_Blueprint
+{
+	GENERATED_BODY()
+
+	/** The affected snapshot */
+	UPROPERTY(BlueprintReadOnly, Category = "Level Snapshots")
+	TObjectPtr<ULevelSnapshot> Snapshot;
+};
+
+USTRUCT(DisplayName = "PreTakeSnapshotEventData")
+struct LEVELSNAPSHOTS_API FPreTakeSnapshotEventData_Blueprint : public FLevelSnapshotEvent_Blueprint { GENERATED_BODY() };
+USTRUCT(DisplayName = "PostTakeSnapshotEventData")
+struct LEVELSNAPSHOTS_API FPostTakeSnapshotEventData_Blueprint : public FLevelSnapshotEvent_Blueprint { GENERATED_BODY() };
+
+USTRUCT(DisplayName = "PreApplySnapshotEventData")
+struct LEVELSNAPSHOTS_API FPreApplySnapshotEventData_Blueprint : public FLevelSnapshotEvent_Blueprint { GENERATED_BODY() };
+USTRUCT(DisplayName = "PostApplySnapshotEventData")
+struct LEVELSNAPSHOTS_API FPostApplySnapshotEventData_Blueprint : public FLevelSnapshotEvent_Blueprint { GENERATED_BODY() };
+
+/** Exposes some native Level Snapshots events to Blueprint scripting.  */
+UCLASS()
+class LEVELSNAPSHOTS_API ULevelSnapshotsEngineSubsystem : public UEngineSubsystem
+{
+	GENERATED_BODY()
+public:
+
+	DECLARE_DYNAMIC_MULTICAST_DELEGATE_OneParam(FPreTakeSnapshotEvent, const FPreTakeSnapshotEventData_Blueprint&, EventData);
+	/** Called before a level snapshot captures the world's data. */
+	UPROPERTY(BlueprintAssignable, Category = "Level Snapshots")
+	FPreTakeSnapshotEvent OnPreTakeSnapshot;
+
+	DECLARE_DYNAMIC_MULTICAST_DELEGATE_OneParam(FPostTakeSnapshotEvent, const FPostTakeSnapshotEventData_Blueprint&, EventData);
+	/** Called after a level snapshot captures the world's data. */
+	UPROPERTY(BlueprintAssignable, Category = "Level Snapshots")
+	FPostTakeSnapshotEvent OnPostTakeSnapshot;
+
+	DECLARE_DYNAMIC_MULTICAST_DELEGATE_OneParam(FPreApplySnapshotEvent, const FPreApplySnapshotEventData_Blueprint&, EventData);
+	/** Called before a level snapshot is applied to a world. */
+	UPROPERTY(BlueprintAssignable, Category = "Level Snapshots")
+	FPreApplySnapshotEvent OnPreApplySnapshot;
+
+	DECLARE_DYNAMIC_MULTICAST_DELEGATE_OneParam(FPostApplySnapshotEvent, const FPostApplySnapshotEventData_Blueprint&, EventData);
+	/** Called after a level snapshot was applied to a world. */
+	UPROPERTY(BlueprintAssignable, Category = "Level Snapshots")
+	FPostApplySnapshotEvent OnPostApplySnapshot;
+
+	//~ Begin USubsystem Interface
+	virtual void Initialize(FSubsystemCollectionBase& Collection) override;
+	virtual void Deinitialize() override;
+	//~ End USubsystem Interface
+
+private:
+
+	/** Exposes some events to Blueprints. Created by Initialize and Reset() by Deinitialize. */
+	TSharedPtr<UE::LevelSnapshots::IRestorationListener> RestorationListener;
+
+	void HandleOnPreTakeSnapshot(const UE::LevelSnapshots::FPreTakeSnapshotEventData& EventData);
+	void HandleOnPostTakeSnapshot(const UE::LevelSnapshots::FPostTakeSnapshotEventData& EventData);
+};
diff --git a/Engine/Plugins/VirtualProduction/LevelSnapshots/Source/LevelSnapshots/Public/Restorability/Interfaces/IRestorationListener.h b/Engine/Plugins/VirtualProduction/LevelSnapshots/Source/LevelSnapshots/Public/Restorability/Interfaces/IRestorationListener.h
--- a/Engine/Plugins/VirtualProduction/LevelSnapshots/Source/LevelSnapshots/Public/Restorability/Interfaces/IRestorationListener.h
+++ b/Engine/Plugins/VirtualProduction/LevelSnapshots/Source/LevelSnapshots/Public/Restorability/Interfaces/IRestorationListener.h
@@ -7,6 +7,7 @@
 #include "Engine/World.h"
 #include "Templates/SubclassOf.h"
 
+class ULevelSnapshot;
 class AActor;
 class UActorComponent;
 
@@ -15,28 +16,37 @@ struct FPropertySelectionMap;
 
 namespace UE::LevelSnapshots
 {
-	struct FApplySnapshotParams
+	struct FRestorationListenerBaseParams
 	{
 		/** All of the user's selected properties */
 		const FPropertySelectionMap& SelectedProperties;
-
-		FApplySnapshotParams(const FPropertySelectionMap& SelectedProperties)
+		
+		FRestorationListenerBaseParams(const FPropertySelectionMap& SelectedProperties)
 			: SelectedProperties(SelectedProperties)
 		{}
 	};
 
-	struct FPreRemoveActorParams : FApplySnapshotParams
+	struct FApplySnapshotSharedParams : FRestorationListenerBaseParams
+	{
+		ULevelSnapshot& LevelSnapshot;
+		FApplySnapshotSharedParams(ULevelSnapshot& LevelSnapshot, const FPropertySelectionMap& SelectedProperties)
+			: FRestorationListenerBaseParams(SelectedProperties)
+			, LevelSnapshot(LevelSnapshot)
+		{}
+	};
+	struct FPreApplySnapshotParams : FApplySnapshotSharedParams {};
+	struct FPostApplySnapshotParams : FApplySnapshotSharedParams {};
+
+	struct FPreRemoveActorParams : FRestorationListenerBaseParams
 	{
 		AActor* ActorToRemove;
 
 		FPreRemoveActorParams(const FPropertySelectionMap& SelectedProperties, AActor* ActorToRemove)
-			: FApplySnapshotParams(SelectedProperties),
+			: FRestorationListenerBaseParams(SelectedProperties),
 			  ActorToRemove(ActorToRemove)
 		{}
 	};
-
-	struct FPostRemoveActorsParams : FApplySnapshotParams
-	{};
+	struct FPostRemoveActorsParams : FRestorationListenerBaseParams {};
 	
 	struct FApplySnapshotPropertiesParams
 	{
@@ -132,10 +142,10 @@ namespace UE::LevelSnapshots
 	public:
 
 		/** Called before the snapshot is applied. Except for the changes made by this event, the world has not been changed, yet. */
-		virtual void PreApplySnapshot(const FApplySnapshotParams& Params) {}
+		virtual void PreApplySnapshot(const FPreApplySnapshotParams& Params) {}
 		
 		/** Called after all snapshot data was applied. The world is nearly completely changed (the only thing remaining is executing this event). */
-		virtual void PostApplySnapshot(const FApplySnapshotParams& Params) {}
+		virtual void PostApplySnapshot(const FPostApplySnapshotParams& Params) {}
 
 
 		
