Fixing debug build error.

diff --git a/Engine/Plugins/VirtualProduction/LevelSnapshots/Source/LevelSnapshots/Private/Archive/ClassDefaults/ApplyClassDefaulDataArchive.cpp b/Engine/Plugins/VirtualProduction/LevelSnapshots/Source/LevelSnapshots/Private/Archive/ClassDefaults/ApplyClassDefaulDataArchive.cpp
--- a/Engine/Plugins/VirtualProduction/LevelSnapshots/Source/LevelSnapshots/Private/Archive/ClassDefaults/ApplyClassDefaulDataArchive.cpp
+++ b/Engine/Plugins/VirtualProduction/LevelSnapshots/Source/LevelSnapshots/Private/Archive/ClassDefaults/ApplyClassDefaulDataArchive.cpp
@@ -3,6 +3,7 @@
 #include "Archive/ClassDefaults/ApplyClassDefaulDataArchive.h"
 
 #include "Data/Util/ObjectDependencyUtil.h"
+#include "LevelSnapshotsLog.h"
 #include "LevelSnapshotsModule.h"
 #include "ObjectSnapshotData.h"
 #include "SelectionSet.h"
