LevelSnapshots now correctly removes PCG actors when they are placed in a sublevel

diff --git a/Engine/Plugins/VirtualProduction/LevelSnapshots/Source/LevelSnapshots/LevelSnapshots.Build.cs b/Engine/Plugins/VirtualProduction/LevelSnapshots/Source/LevelSnapshots/LevelSnapshots.Build.cs
--- a/Engine/Plugins/VirtualProduction/LevelSnapshots/Source/LevelSnapshots/LevelSnapshots.Build.cs
+++ b/Engine/Plugins/VirtualProduction/LevelSnapshots/Source/LevelSnapshots/LevelSnapshots.Build.cs
@@ -49,7 +49,8 @@ public class LevelSnapshots : ModuleRules
 			PrivateDependencyModuleNames.AddRange( new string[]
 			{
 				"UnrealEd", // This is needed for undo / redo system in editor
-				"Settings"	// For registering settings
+				"Settings",	// For registering settings
+				"Projects"	// For determining enabled plugins
 			});
 		}
 	}
diff --git a/Engine/Plugins/VirtualProduction/LevelSnapshots/Source/LevelSnapshots/Private/Data/Util/WorldDataUtil.cpp b/Engine/Plugins/VirtualProduction/LevelSnapshots/Source/LevelSnapshots/Private/Data/Util/WorldDataUtil.cpp
--- a/Engine/Plugins/VirtualProduction/LevelSnapshots/Source/LevelSnapshots/Private/Data/Util/WorldDataUtil.cpp
+++ b/Engine/Plugins/VirtualProduction/LevelSnapshots/Source/LevelSnapshots/Private/Data/Util/WorldDataUtil.cpp
@@ -331,23 +331,40 @@ namespace UE::LevelSnapshots::Private::Internal
 	
 		// Not sure whether needed. "DELETE" command does in UUnrealEdEngine::Exec_Actor ...
 		FEditorDelegates::OnDeleteActorsBegin.Broadcast();
-	
+		
 		FLevelSnapshotsModule& Module = FLevelSnapshotsModule::GetInternalModuleInstance();
 		TArray<AActor*> ActorsToDelete;
+		// Actors returning false from CanDeleteSelectedActor is skipped by DeleteActors: just call EditorDelete on them.
+		TArray<AActor*> UserManagedActors;
 		ActorsToDelete.Reserve(ActorsToDespawn.Num());
 		for (const TWeakObjectPtr<AActor>& ActorToDespawn : ActorsToDespawn)
 		{
-			if (ActorToDespawn.IsValid())
+			if (!ActorToDespawn.IsValid())
+			{
+				continue;
+			}
+
+			Module.OnPreRemoveActor({ PropertiesToSerialize, ActorToDespawn.Get() });
+			if (FText Dummy; ActorToDespawn->CanDeleteSelectedActor(Dummy))
 			{
-				Module.OnPreRemoveActor({ PropertiesToSerialize, ActorToDespawn.Get() });
 				ActorsToDelete.Add(ActorToDespawn.Get());
 			}
+			else
+			{
+				UserManagedActors.Add(ActorToDespawn.Get());
+			}
 		}
 
+		// DeleteActors handles advanced clean-up cases, such as proper transaction, group actors, level script, etc.
 		constexpr bool bVerifyDeletionCanHappen = true;
 		constexpr bool bWarnAboutReferences = false;
 		GUnrealEd->DeleteActors(ActorsToDelete, WorldToApplyTo, GUnrealEd->GetSelectedActors()->GetElementSelectionSet(), bVerifyDeletionCanHappen, bWarnAboutReferences, bWarnAboutReferences);
-
+		// EditorDestroyActor for any actors returning false from CanDeleteSelectedActor (these are usually advanced tool specific actors, like level instances, PCG actors, etc.)
+		for (AActor* UserManagedActor : UserManagedActors)
+		{
+			UserManagedActor->GetWorld()->EditorDestroyActor(UserManagedActor, true);
+		}
+		
 		// ... and call the end event like in UUnrealEdEngine
 		FEditorDelegates::OnDeleteActorsEnd.Broadcast();
 #else
diff --git a/Engine/Plugins/VirtualProduction/LevelSnapshots/Source/LevelSnapshots/Private/LevelSnapshotsModule.cpp b/Engine/Plugins/VirtualProduction/LevelSnapshots/Source/LevelSnapshots/Private/LevelSnapshotsModule.cpp
--- a/Engine/Plugins/VirtualProduction/LevelSnapshots/Source/LevelSnapshots/Private/LevelSnapshotsModule.cpp
+++ b/Engine/Plugins/VirtualProduction/LevelSnapshots/Source/LevelSnapshots/Private/LevelSnapshotsModule.cpp
@@ -10,8 +10,8 @@
 
 #include "Algo/AllOf.h"
 #include "Components/ActorComponent.h"
-#include "EngineUtils.h"
 #include "GameFramework/Actor.h"
+#include "Misc/CoreDelegates.h"
 #include "Modules/ModuleManager.h"
 #if WITH_EDITOR
 #include "ISettingsModule.h"
@@ -66,8 +66,12 @@ void UE::LevelSnapshots::Private::FLevelSnapshotsModule::StartupModule()
 	);
 	RegisterRestorabilityOverrider(ClassSkipper);
 	
-	// Add support for engine features that require specialized case handling
-	EngineTypesRestorationFence::RegisterSpecialEngineTypeSupport(*this);
+	// Add support for engine features that require specialized case handling.
+	// Since we may depend on dynamic class look-ups (e.g. FPCGRestoration), this is done in OnPostEngineInit.
+	FCoreDelegates::OnPostEngineInit.AddLambda([this]()
+	{
+		EngineTypesRestorationFence::RegisterSpecialEngineTypeSupport(*this);
+	});
 
 #if WITH_EDITOR
 	ISettingsModule& SettingsModule = FModuleManager::LoadModuleChecked<ISettingsModule>("Settings");
diff --git a/Engine/Plugins/VirtualProduction/LevelSnapshots/Source/LevelSnapshots/Private/SpecialSupport/EngineTypesRestorationFence.cpp b/Engine/Plugins/VirtualProduction/LevelSnapshots/Source/LevelSnapshots/Private/SpecialSupport/EngineTypesRestorationFence.cpp
--- a/Engine/Plugins/VirtualProduction/LevelSnapshots/Source/LevelSnapshots/Private/SpecialSupport/EngineTypesRestorationFence.cpp
+++ b/Engine/Plugins/VirtualProduction/LevelSnapshots/Source/LevelSnapshots/Private/SpecialSupport/EngineTypesRestorationFence.cpp
@@ -5,17 +5,18 @@
 #include "ActorGroupRestoration.h"
 #include "ActorLabelRestoration.h"
 #include "AttachParentShowsTransformProperties.h"
+#include "CollisionRestoration.h"
+#include "GridPlacementRestoration.h"
 #include "LevelInstanceRestoration.h"
+#include "PCGRestoration.h"
 #include "SpecialActorPropertySupport.h"
-#include "SpecialSupport/CollisionRestoration.h"
-#include "SpecialSupport/GridPlacementRestoration.h"
 
 #include "Algo/Transform.h"
 #include "Components/SkeletalMeshComponent.h"
-#include "EngineUtils.h"
 #include "MaterialInstanceSupport.h"
 #include "Engine/Brush.h"
 #include "GameFramework/Actor.h"
+#include "GameFramework/WorldSettings.h"
 #include "Materials/MaterialInstance.h"
 
 namespace UE::LevelSnapshots::Private::EngineTypesRestorationFence
@@ -130,5 +131,6 @@ namespace UE::LevelSnapshots::Private::EngineTypesRestorationFence
 		LevelInstanceRestoration::Register(Module);
 		MaterialInstanceSupport::Register(Module);
 		ActorLabelRestoration::Register(Module);
+		PCGRestoration::Register(Module);
 	}
 }
diff --git a/Engine/Plugins/VirtualProduction/LevelSnapshots/Source/LevelSnapshots/Private/SpecialSupport/PCGRestoration.cpp b/Engine/Plugins/VirtualProduction/LevelSnapshots/Source/LevelSnapshots/Private/SpecialSupport/PCGRestoration.cpp
new file mode 100644
--- /dev/null
+++ b/Engine/Plugins/VirtualProduction/LevelSnapshots/Source/LevelSnapshots/Private/SpecialSupport/PCGRestoration.cpp
@@ -0,0 +1,46 @@
+// Copyright Epic Games, Inc. All Rights Reserved.
+
+#include "PCGRestoration.h"
+
+#include "Restorability/Interfaces/ISnapshotRestorabilityOverrider.h"
+
+#include "Components/StaticMeshComponent.h"
+#if WITH_EDITOR
+#include "Interfaces/IPluginManager.h"
+#endif
+
+namespace UE::LevelSnapshots::Private::PCGRestoration
+{
+	/** Prevents restoration of components added to APCGVolume because they are refreshed and re-generated anyways. */
+	class FPCGRestoration : public ISnapshotRestorabilityOverrider
+	{
+		/** Class of APCGVolume. Nullptr if PCG plugin is not enabled. */
+		UClass* PCGVolumeClass;
+	public:
+
+		FPCGRestoration()
+			: PCGVolumeClass(FindObject<UClass>(nullptr, TEXT("/Script/PCG.PCGVolume")))
+		{
+#if WITH_EDITOR
+			const TSharedPtr<IPlugin> Plugin = IPluginManager::Get().FindPlugin(TEXT("PCG"));
+			ensureMsgf(!Plugin || !Plugin->IsEnabled() || PCGVolumeClass, TEXT("PCG plugin is enabled but failed to find APCGVolume class. Was the class renamed?"));
+#endif
+		}
+		
+		virtual ERestorabilityOverride IsComponentDesirableForCapture(const UActorComponent* Component) override
+		{
+			const bool bIsOnPCGVolumne = PCGVolumeClass
+				// Do not restore generated mesh components attached to APCGVolume
+				&& Component->IsA<UStaticMeshComponent>()
+				&& Component->GetOwner()->GetClass() == PCGVolumeClass;
+			return bIsOnPCGVolumne
+				? ERestorabilityOverride::Disallow
+				: ERestorabilityOverride::DoNotCare;
+		}
+	};
+	
+	void Register(FLevelSnapshotsModule& Module)
+	{
+		Module.RegisterRestorabilityOverrider(MakeShared<FPCGRestoration>());
+	}
+}
diff --git a/Engine/Plugins/VirtualProduction/LevelSnapshots/Source/LevelSnapshots/Private/SpecialSupport/PCGRestoration.h b/Engine/Plugins/VirtualProduction/LevelSnapshots/Source/LevelSnapshots/Private/SpecialSupport/PCGRestoration.h
new file mode 100644
--- /dev/null
+++ b/Engine/Plugins/VirtualProduction/LevelSnapshots/Source/LevelSnapshots/Private/SpecialSupport/PCGRestoration.h
@@ -0,0 +1,11 @@
+// Copyright Epic Games, Inc. All Rights Reserved.
+
+#pragma once
+
+#include "LevelSnapshotsModule.h"
+
+/** PCG = Procedural Content Generation */
+namespace UE::LevelSnapshots::Private::PCGRestoration
+{
+	void Register(FLevelSnapshotsModule& Module);
+}
diff --git a/Engine/Plugins/VirtualProduction/LevelSnapshots/Source/LevelSnapshots/Private/Tests/RegressionTests.cpp b/Engine/Plugins/VirtualProduction/LevelSnapshots/Source/LevelSnapshots/Private/Tests/RegressionTests.cpp
--- a/Engine/Plugins/VirtualProduction/LevelSnapshots/Source/LevelSnapshots/Private/Tests/RegressionTests.cpp
+++ b/Engine/Plugins/VirtualProduction/LevelSnapshots/Source/LevelSnapshots/Private/Tests/RegressionTests.cpp
@@ -7,6 +7,7 @@
 
 #include "Engine/World.h"
 #include "EngineUtils.h"
+#include "Algo/Count.h"
 #include "Misc/AutomationTest.h"
 
 // Bug fixes should generally be tested. Put tests for bug fixes here.
@@ -250,4 +251,32 @@ namespace UE::LevelSnapshots::Private::Tests
 
 		return true;
 	}
+
+#if WITH_EDITOR
+	/** Tests that actors implementing AActor::CanDeleteSelectedActor and AActor::IsUserManaged are properly removed. */
+	IMPLEMENT_SIMPLE_AUTOMATION_TEST(FUserManagedActors, "VirtualProduction.LevelSnapshots.Snapshot.Regression.UserManagedActors", (EAutomationTestFlags::ApplicationContextMask | EAutomationTestFlags::EngineFilter));
+	bool FUserManagedActors::RunTest(const FString& Parameters)
+	{
+		int32 NumActorsBeforeModify;
+		FSnapshotTestRunner()
+			.TakeSnapshot()
+			.ModifyWorld([&](UWorld* World)
+			{
+				NumActorsBeforeModify = Algo::CountIf(World->PersistentLevel->Actors, [](AActor* Actor){ return Actor != nullptr; });
+				
+				ASnapshotTestActor* NotDeletableActor = ASnapshotTestActor::Spawn(World);
+				ASnapshotTestActor* NotUserManagedActor = ASnapshotTestActor::Spawn(World);
+				NotDeletableActor->DeleteSelectedActorMode = ASnapshotTestActor::ECanDeleteMode::No;
+				NotUserManagedActor->bIsUserManaged = false;
+			})
+			.ApplySnapshot()
+			.ModifyWorld([&](UWorld* World)
+			{
+				const int32 NumActorsAfterRestore = Algo::CountIf(World->PersistentLevel->Actors, [](AActor* Actor){ return Actor != nullptr; });
+				TestEqual(TEXT("CanDeleteSelectedActor & IsUserManaged"), NumActorsAfterRestore, NumActorsBeforeModify);
+			});
+
+		return true;
+	}
+#endif
 }
\ No newline at end of file
diff --git a/Engine/Plugins/VirtualProduction/LevelSnapshots/Source/LevelSnapshots/Private/Tests/Types/SnapshotTestActor.h b/Engine/Plugins/VirtualProduction/LevelSnapshots/Source/LevelSnapshots/Private/Tests/Types/SnapshotTestActor.h
--- a/Engine/Plugins/VirtualProduction/LevelSnapshots/Source/LevelSnapshots/Private/Tests/Types/SnapshotTestActor.h
+++ b/Engine/Plugins/VirtualProduction/LevelSnapshots/Source/LevelSnapshots/Private/Tests/Types/SnapshotTestActor.h
@@ -93,6 +93,28 @@ public:
 	//~ Begin UObject Interface
 	virtual void PostInitProperties() override;
 	//~ End UObject Interface
+
+#if WITH_EDITOR
+	enum class ECanDeleteMode
+	{
+		CallSuper,
+		Yes,
+		No
+	} DeleteSelectedActorMode = ECanDeleteMode::CallSuper;
+	bool bIsUserManaged = true;
+	
+	virtual bool CanDeleteSelectedActor(FText& OutReason) const override
+	{
+		switch (DeleteSelectedActorMode)
+		{
+		case ECanDeleteMode::Yes: return true;
+		case ECanDeleteMode::No: return false;
+		default:
+			return Super::CanDeleteSelectedActor(OutReason);
+		}
+	}
+	virtual bool IsUserManaged() const override  { return bIsUserManaged; }
+#endif
 	
 	
 	/******************** Skipped properties  ********************/
