Fix Level Snapshots not restoring custom subobject properties that were equal to class archetype value at time of taking the snapshot

diff --git a/Engine/Plugins/VirtualProduction/LevelSnapshots/Source/LevelSnapshots/Private/Archive/ApplySnapshotToEditorArchive.cpp b/Engine/Plugins/VirtualProduction/LevelSnapshots/Source/LevelSnapshots/Private/Archive/ApplySnapshotToEditorArchive.cpp
--- a/Engine/Plugins/VirtualProduction/LevelSnapshots/Source/LevelSnapshots/Private/Archive/ApplySnapshotToEditorArchive.cpp
+++ b/Engine/Plugins/VirtualProduction/LevelSnapshots/Source/LevelSnapshots/Private/Archive/ApplySnapshotToEditorArchive.cpp
@@ -36,7 +36,7 @@ void UE::LevelSnapshots::Private::FApplySnapshotToEditorArchive::ApplyToExisting
 #endif
 	
 	// 1. Serialize archetype first to handle the case were the archetype has changed properties since the snapshot was taken
-	if (ClassIndex) // Sometimes not available, e.g. for custom subobjects
+	if (ClassIndex)
 	{
 		const FSubobjectArchetypeFallbackInfo ClassFallbackInfo{ InOriginalObject->GetOuter(), InOriginalObject->GetFName(), InOriginalObject->GetFlags() };
 		SerializeSelectedClassDefaultsIntoSubobject(InOriginalObject, InSharedData, *ClassIndex, Cache, ClassFallbackInfo, *Selection);
diff --git a/Engine/Plugins/VirtualProduction/LevelSnapshots/Source/LevelSnapshots/Private/Data/CustomSerialization/CustomObjectSerializationWrapper.cpp b/Engine/Plugins/VirtualProduction/LevelSnapshots/Source/LevelSnapshots/Private/Data/CustomSerialization/CustomObjectSerializationWrapper.cpp
--- a/Engine/Plugins/VirtualProduction/LevelSnapshots/Source/LevelSnapshots/Private/Data/CustomSerialization/CustomObjectSerializationWrapper.cpp
+++ b/Engine/Plugins/VirtualProduction/LevelSnapshots/Source/LevelSnapshots/Private/Data/CustomSerialization/CustomObjectSerializationWrapper.cpp
@@ -115,7 +115,17 @@ namespace UE::LevelSnapshots::Private::Internal
 						[&WorldData, OriginalPath = MetaData->GetOriginalPath()](){ return UE::LevelSnapshots::Private::FindCustomSubobjectData(WorldData, OriginalPath);} );
 				
 					FCustomSerializationData* SerializationData = SerializationDataGetter();
-					FApplySnapshotToEditorArchive::ApplyToExistingEditorWorldObject(SerializationData->Subobjects[i], WorldData, Cache, EditorSubobject, SnapshotSubobject, SelectionMap);
+					FCustomSubbjectSerializationData& DataToRestore = SerializationData->Subobjects[i];
+					// ClassIndex is only saved starting 5.3. Previous data will default to INDEX_NONE. We don't want to cause a bunch of warning logs by passing in a set class index.
+					const TOptional<FClassDataIndex> ClassIndex = DataToRestore.ClassIndex == INDEX_NONE ? TOptional<uint32>{} : static_cast<FClassDataIndex>(DataToRestore.ClassIndex);
+					FApplySnapshotToEditorArchive::ApplyToExistingEditorWorldObject(DataToRestore,
+						WorldData,
+						Cache,
+						EditorSubobject,
+						SnapshotSubobject,
+						SelectionMap,
+						ClassIndex
+						);
 					CustomSerializer->OnPostSerializeEditorSubobject(EditorSubobject, *MetaData, SerializationDataReader);
 					continue;
 				}
diff --git a/Engine/Plugins/VirtualProduction/LevelSnapshots/Source/LevelSnapshots/Private/Data/CustomSerialization/CustomSerializationDataManager.cpp b/Engine/Plugins/VirtualProduction/LevelSnapshots/Source/LevelSnapshots/Private/Data/CustomSerialization/CustomSerializationDataManager.cpp
--- a/Engine/Plugins/VirtualProduction/LevelSnapshots/Source/LevelSnapshots/Private/Data/CustomSerialization/CustomSerializationDataManager.cpp
+++ b/Engine/Plugins/VirtualProduction/LevelSnapshots/Source/LevelSnapshots/Private/Data/CustomSerialization/CustomSerializationDataManager.cpp
@@ -9,6 +9,7 @@
 #include "LevelSnapshotsLog.h"
 #include "SnapshotCustomVersion.h"
 #include "TakeWorldObjectSnapshotArchive.h"
+#include "Util/ClassDataUtil.h"
 
 #include "Serialization/MemoryWriter.h"
 #include "Serialization/MemoryReader.h"
@@ -260,6 +261,7 @@ int32 UE::LevelSnapshots::Private::FCustomSerializationDataWriter::AddSubobjectS
 
 	FCustomSubbjectSerializationData SubobjectData;
 	SubobjectData.ObjectPathIndex = AddObjectDependency(WorldData_ReadWrite, Subobject, false);
+	SubobjectData.ClassIndex = AddClassArchetype(WorldData_ReadWrite, Subobject);
 	const int32 SubobjectIndex = SerializationData->Subobjects.Emplace(
 		MoveTemp(SubobjectData) // Not profiled
 		);
diff --git a/Engine/Plugins/VirtualProduction/LevelSnapshots/Source/LevelSnapshots/Private/Tests/CustomObjectSerializationTests.cpp b/Engine/Plugins/VirtualProduction/LevelSnapshots/Source/LevelSnapshots/Private/Tests/CustomObjectSerializationTests.cpp
--- a/Engine/Plugins/VirtualProduction/LevelSnapshots/Source/LevelSnapshots/Private/Tests/CustomObjectSerializationTests.cpp
+++ b/Engine/Plugins/VirtualProduction/LevelSnapshots/Source/LevelSnapshots/Private/Tests/CustomObjectSerializationTests.cpp
@@ -770,7 +770,6 @@ namespace UE::LevelSnapshots::Private::Tests
 	IMPLEMENT_SIMPLE_AUTOMATION_TEST(FRestoresCustomSubobjectWhenActorRecreated, "VirtualProduction.LevelSnapshots.Snapshot.CustomObjectSerialization.RestoresCustomSubobjectWhenActorRecreated", (EAutomationTestFlags::ApplicationContextMask | EAutomationTestFlags::EngineFilter));
 	bool FRestoresCustomSubobjectWhenActorRecreated::RunTest(const FString& Parameters)
 	{
-		// Handle registering and unregistering of custom serializer
 		// Handle registering and unregistering of custom serializer
 		const TCustomObjectSerializerContext<FInstancedOnlySubobjectCustomObjectSerializer> Stub =
 			FInstancedOnlySubobjectCustomObjectSerializer::Make();
@@ -804,6 +803,37 @@ namespace UE::LevelSnapshots::Private::Tests
 				TestEqual(TEXT("Default Subobject was restored"), RestoredActor->InstancedOnlySubobject_DefaultSubobject->FloatProperty, 12345.f);
 			});
 		
+		return true;
+	}
+
+	/** Tests that class archetype data is saved for custom subobjects. That means if a snapshot is taken on an unmodified custom subobject and then changed, it should restore. Fix for UE-168414. */
+	IMPLEMENT_SIMPLE_AUTOMATION_TEST(FRestoreCustomSubobjectPropertiesEqualToArchetypeValue, "VirtualProduction.LevelSnapshots.Snapshot.CustomObjectSerialization.RestorePropertiesEqualToArchetypeValue", (EAutomationTestFlags::ApplicationContextMask | EAutomationTestFlags::EngineFilter));
+	bool FRestoreCustomSubobjectPropertiesEqualToArchetypeValue::RunTest(const FString& Parameters)
+	{
+		// Handle registering and unregistering of custom serializer
+		const TCustomObjectSerializerContext<FInstancedOnlySubobjectCustomObjectSerializer> Stub =
+			FInstancedOnlySubobjectCustomObjectSerializer::Make();
+
+		FSnapshotTestRunner()
+			.ModifyWorld([&](UWorld* World)
+			{
+				Stub.GetCustomSerializer()->TestActor = ASnapshotTestActor::Spawn(World, "RecreatedActor");
+			})
+			.TakeSnapshot()
+	
+			.ModifyWorld([&](UWorld* World)
+			{
+				Stub.GetCustomSerializer()->TestActor->InstancedOnlySubobject_DefaultSubobject->FloatProperty = 5.f;
+			})
+			.ApplySnapshot()
+		
+			.ModifyWorld([&](UWorld* World)
+			{
+				TestEqual(TEXT("Restore property that was equal to archetype value at time of taking snapshot"),
+					Stub.GetCustomSerializer()->TestActor->InstancedOnlySubobject_DefaultSubobject->FloatProperty,
+					0.f);
+			});
+		
 		return true;
 	} 
 }
diff --git a/Engine/Plugins/VirtualProduction/LevelSnapshots/Source/LevelSnapshots/Private/Tests/Types/SnapshotTestActor.h b/Engine/Plugins/VirtualProduction/LevelSnapshots/Source/LevelSnapshots/Private/Tests/Types/SnapshotTestActor.h
--- a/Engine/Plugins/VirtualProduction/LevelSnapshots/Source/LevelSnapshots/Private/Tests/Types/SnapshotTestActor.h
+++ b/Engine/Plugins/VirtualProduction/LevelSnapshots/Source/LevelSnapshots/Private/Tests/Types/SnapshotTestActor.h
@@ -18,10 +18,10 @@ class USubSubobject : public UObject
 public:
 
 	UPROPERTY(EditAnywhere, Category = "Level Snapshots")
-	int32 IntProperty;
+	int32 IntProperty = 0;
 
 	UPROPERTY(EditAnywhere, Category = "Level Snapshots")
-	float FloatProperty;
+	float FloatProperty = 0;
 	
 };
 
@@ -34,10 +34,10 @@ public:
 	USubobject();
 	
 	UPROPERTY(EditAnywhere, Category = "Level Snapshots")
-	int32 IntProperty;
+	int32 IntProperty = 0;
 
 	UPROPERTY(EditAnywhere, Category = "Level Snapshots")
-	float FloatProperty;
+	float FloatProperty = 0;
 
 	UPROPERTY(EditAnywhere, Instanced, Category = "Level Snapshots")
 	TObjectPtr<USubSubobject> NestedChild;
@@ -59,10 +59,10 @@ public:
 	USnapshotTestComponent();
 	
 	UPROPERTY(EditAnywhere, Category = "Level Snapshots")
-	int32 IntProperty;
+	int32 IntProperty = 0;
 
 	UPROPERTY(EditAnywhere, Category = "Level Snapshots")
-	float FloatProperty;
+	float FloatProperty = 0;
 
 	UPROPERTY(Instanced)
 	TObjectPtr<USubobject> Subobject;
@@ -126,7 +126,7 @@ public:
 	int32 TransientProperty = 200;
 
 	UPROPERTY(EditAnywhere, Category = "Level Snapshots")
-	int32 IntProperty;
+	int32 IntProperty = 0;
 
 
 	
diff --git a/Engine/Plugins/VirtualProduction/LevelSnapshots/Source/LevelSnapshots/Public/Data/ClassSnapshotData.h b/Engine/Plugins/VirtualProduction/LevelSnapshots/Source/LevelSnapshots/Public/Data/ClassSnapshotData.h
--- a/Engine/Plugins/VirtualProduction/LevelSnapshots/Source/LevelSnapshots/Public/Data/ClassSnapshotData.h
+++ b/Engine/Plugins/VirtualProduction/LevelSnapshots/Source/LevelSnapshots/Public/Data/ClassSnapshotData.h
@@ -29,10 +29,7 @@ struct LEVELSNAPSHOTS_API FClassSnapshotData : public FObjectSnapshotData
 {
 	GENERATED_BODY()
 
-	/**
-	 * The class getting saved.
-	 * Note:
-	 */
+	/** The class getting saved. */
 	UPROPERTY()
 	FSoftClassPath ClassPath;
 
diff --git a/Engine/Plugins/VirtualProduction/LevelSnapshots/Source/LevelSnapshots/Public/Data/CustomSerializationData.h b/Engine/Plugins/VirtualProduction/LevelSnapshots/Source/LevelSnapshots/Public/Data/CustomSerializationData.h
--- a/Engine/Plugins/VirtualProduction/LevelSnapshots/Source/LevelSnapshots/Public/Data/CustomSerializationData.h
+++ b/Engine/Plugins/VirtualProduction/LevelSnapshots/Source/LevelSnapshots/Public/Data/CustomSerializationData.h
@@ -15,6 +15,10 @@ struct LEVELSNAPSHOTS_API FCustomSubbjectSerializationData : public FObjectSnaps
 	UPROPERTY()
 	int32 ObjectPathIndex = INDEX_NONE;
 
+	/** Valid index to FWorldSnapshotData::ClassData. Use to lookup class and archetype data. */
+	UPROPERTY()
+	int32 ClassIndex = INDEX_NONE;
+
 	/* Additional custom data saved by ICustomSnapshotSerializationData */
 	UPROPERTY()
 	TArray<uint8> SubobjectAnnotationData;
