Level Snapshots now handles transient AssetUserData being saved

diff --git a/Engine/Plugins/VirtualProduction/LevelSnapshots/Source/LevelSnapshots/Private/Archive/TakeWorldObjectSnapshotArchive.cpp b/Engine/Plugins/VirtualProduction/LevelSnapshots/Source/LevelSnapshots/Private/Archive/TakeWorldObjectSnapshotArchive.cpp
--- a/Engine/Plugins/VirtualProduction/LevelSnapshots/Source/LevelSnapshots/Private/Archive/TakeWorldObjectSnapshotArchive.cpp
+++ b/Engine/Plugins/VirtualProduction/LevelSnapshots/Source/LevelSnapshots/Private/Archive/TakeWorldObjectSnapshotArchive.cpp
@@ -3,6 +3,7 @@
 #include "TakeWorldObjectSnapshotArchive.h"
 
 #include "LevelSnapshotsLog.h"
+#include "LevelSnapshotsModule.h"
 #include "ObjectSnapshotData.h"
 #include "WorldSnapshotData.h"
 
@@ -12,10 +13,15 @@
 void UE::LevelSnapshots::Private::FTakeWorldObjectSnapshotArchive::TakeSnapshot(FObjectSnapshotData& InObjectData, FWorldSnapshotData& InSharedData, UObject* InOriginalObject)
 {
 	check(InOriginalObject);
-	
-	FTakeWorldObjectSnapshotArchive Archive(InObjectData, InSharedData, InOriginalObject);
-	InOriginalObject->Serialize(Archive);
-	InObjectData.ObjectFlags = InOriginalObject->GetFlags();
+
+	FLevelSnapshotsModule& Module = FLevelSnapshotsModule::GetInternalModuleInstance();
+	Module.OnPreTakeObjectSnapshot({ InOriginalObject });
+	{
+		FTakeWorldObjectSnapshotArchive Archive(InObjectData, InSharedData, InOriginalObject);
+		InOriginalObject->Serialize(Archive);
+		InObjectData.ObjectFlags = InOriginalObject->GetFlags();
+	}
+	Module.OnPostTakeObjectSnapshot({ InOriginalObject });
 }
 
 UE::LevelSnapshots::Private::FTakeWorldObjectSnapshotArchive::FTakeWorldObjectSnapshotArchive(FObjectSnapshotData& InObjectData, FWorldSnapshotData& InSharedData, UObject* InOriginalObject)
diff --git a/Engine/Plugins/VirtualProduction/LevelSnapshots/Source/LevelSnapshots/Private/Data/LevelSnapshot.cpp b/Engine/Plugins/VirtualProduction/LevelSnapshots/Source/LevelSnapshots/Private/Data/LevelSnapshot.cpp
--- a/Engine/Plugins/VirtualProduction/LevelSnapshots/Source/LevelSnapshots/Private/Data/LevelSnapshot.cpp
+++ b/Engine/Plugins/VirtualProduction/LevelSnapshots/Source/LevelSnapshots/Private/Data/LevelSnapshot.cpp
@@ -88,7 +88,7 @@ bool ULevelSnapshot::SnapshotWorld(UWorld* TargetWorld)
 	{
 		return false;
 	}
-	Module.OnPreTakeSnapshot().Broadcast({this});
+	Module.OnPreTakeSnapshot().Broadcast({ this, TargetWorld });
 
 	MapPath = TargetWorld;
 	CaptureTime = FDateTime::UtcNow();
@@ -99,7 +99,7 @@ bool ULevelSnapshot::SnapshotWorld(UWorld* TargetWorld)
 	Cache.InitFor(SerializedData);
 	// Should we delete preexisting actors from the snapshot world?
 
-	Module.OnPostTakeSnapshot().Broadcast({ this });
+	Module.OnPostTakeSnapshot().Broadcast({ this, TargetWorld });
 
 	return true;
 }
diff --git a/Engine/Plugins/VirtualProduction/LevelSnapshots/Source/LevelSnapshots/Private/Filtering/Diffing/ApplySnapshotFilter.cpp b/Engine/Plugins/VirtualProduction/LevelSnapshots/Source/LevelSnapshots/Private/Filtering/Diffing/ApplySnapshotFilter.cpp
--- a/Engine/Plugins/VirtualProduction/LevelSnapshots/Source/LevelSnapshots/Private/Filtering/Diffing/ApplySnapshotFilter.cpp
+++ b/Engine/Plugins/VirtualProduction/LevelSnapshots/Source/LevelSnapshots/Private/Filtering/Diffing/ApplySnapshotFilter.cpp
@@ -72,7 +72,7 @@ void UE::LevelSnapshots::Private::FApplySnapshotFilter::ApplyFilterToFindSelecte
 		const FConditionalScopeLogTime LogTime(ConsoleVariables::CVarLogSelectionMap.GetValueOnAnyThread(), TEXT("Total Time"));
 		
 		FilterActorPair(MapToAddTo);
-		AnalyseComponentProperties(MapToAddTo);
+		AnalyseComponentRestoration(MapToAddTo);
 	}
 }
 
@@ -126,7 +126,7 @@ bool UE::LevelSnapshots::Private::FApplySnapshotFilter::EnsureParametersAreValid
 	return true;
 }
 
-void UE::LevelSnapshots::Private::FApplySnapshotFilter::AnalyseComponentProperties(FPropertySelectionMap& MapToAddTo)
+void UE::LevelSnapshots::Private::FApplySnapshotFilter::AnalyseComponentRestoration(FPropertySelectionMap& MapToAddTo)
 {
 	FAddedAndRemovedComponentInfo ComponentSelection;
 
@@ -178,10 +178,11 @@ void UE::LevelSnapshots::Private::FApplySnapshotFilter::FilterActorPair(FPropert
         FLevelSnapshotPropertyChain(),
         WorldActor->GetClass()
         );
-	
-	AnalyseRootProperties(ActorContext, DeserializedSnapshotActor, WorldActor);
+
+	PreApplyFilter(ActorContext, DeserializedSnapshotActor, WorldActor);
+	TrackChangedRootProperties(ActorContext, DeserializedSnapshotActor, WorldActor);
 	const EFilterObjectPropertiesResult FilterResult = FindAndFilterCustomSubobjectPairs(MapToAddTo, DeserializedSnapshotActor, WorldActor);
-	ExtendAnalysedProperties(ActorContext, DeserializedSnapshotActor, WorldActor);
+	PostApplyFilter(ActorContext, DeserializedSnapshotActor, WorldActor);
 	
 	ActorSelection.SetHasCustomSerializedSubobjects(FilterResult == EFilterObjectPropertiesResult::HasCustomSubobjects);
 	if (!ActorSelection.IsEmpty())
@@ -215,16 +216,16 @@ UE::LevelSnapshots::Private::FApplySnapshotFilter::EPropertySearchResult UE::Lev
         FLevelSnapshotPropertyChain(),
 		WorldSubobject->GetClass()
         );
-	
-	AnalyseRootProperties(ComponentContext, SnapshotSubobject, WorldSubobject);
-	
+
+	PreApplyFilter(ComponentContext, SnapshotSubobject, WorldSubobject);
+	TrackChangedRootProperties(ComponentContext, SnapshotSubobject, WorldSubobject);
 	const EFilterObjectPropertiesResult FilterResult = FindAndFilterCustomSubobjectPairs(MapToAddTo, SnapshotSubobject, WorldSubobject);
 	SubobjectSelection.SetHasCustomSerializedSubobjects(FilterResult == EFilterObjectPropertiesResult::HasCustomSubobjects);
 
 	// Non-component subobjects should always be added so when applying to world and resolving a subobject, we can tell whether it is referenced or just a dead object still existing in memory
 	const bool bAddedProperties = [this, &MapToAddTo, &ComponentContext, SnapshotSubobject, WorldSubobject, &SubobjectSelection]()
 	{
-		ExtendAnalysedProperties(ComponentContext, SnapshotSubobject, WorldSubobject);
+		PostApplyFilter(ComponentContext, SnapshotSubobject, WorldSubobject);
 		if (SubobjectSelection.IsEmpty())
 		{
 			// We don't track components because we can efficiently tell whether they're dead or not by asking the parent actor
@@ -259,7 +260,7 @@ UE::LevelSnapshots::Private::FApplySnapshotFilter::EPropertySearchResult UE::Lev
         Parent.RootClass
         );
 	StructContext.AuthoredPathInformation.Add(StructProperty->GetAuthoredName());
-	return AnalyseStructProperties(StructContext);
+	return TrackChangedPropertiesInStruct(StructContext);
 }
 
 UE::LevelSnapshots::Private::FApplySnapshotFilter::EFilterObjectPropertiesResult UE::LevelSnapshots::Private::FApplySnapshotFilter::FindAndFilterCustomSubobjectPairs(FPropertySelectionMap& MapToAddTo, UObject* SnapshotOwner, UObject* WorldOwner)
@@ -319,7 +320,7 @@ UE::LevelSnapshots::Private::FApplySnapshotFilter::EFilterObjectPropertiesResult
 	return bAtLeastOneSubobjectWasAdded ? EFilterObjectPropertiesResult::HasCustomSubobjects : EFilterObjectPropertiesResult::HasOnlyNormalProperties;
 }
 
-void UE::LevelSnapshots::Private::FApplySnapshotFilter::AnalyseRootProperties(FPropertyContainerContext& ContainerContext, UObject* SnapshotObject, UObject* WorldObject)
+void UE::LevelSnapshots::Private::FApplySnapshotFilter::TrackChangedRootProperties(FPropertyContainerContext& ContainerContext, UObject* SnapshotObject, UObject* WorldObject)
 {
 	FLevelSnapshotsModule& Module = FLevelSnapshotsModule::GetInternalModuleInstance();
 	const FPropertyComparerArray PropertyComparers = Module.GetPropertyComparerForClass(ContainerContext.RootClass);
@@ -343,8 +344,38 @@ void UE::LevelSnapshots::Private::FApplySnapshotFilter::AnalyseRootProperties(FP
 			break;
 		}
 		
-		AnalyseProperty(ContainerContext, *FieldIt, bSkipEqualityTest);
+		TrackPropertyIfChangedRecursive(ContainerContext, *FieldIt, bSkipEqualityTest);
+	}
+}
+
+TOptional<UE::LevelSnapshots::Private::FApplySnapshotFilter::EPropertySearchResult> UE::LevelSnapshots::Private::FApplySnapshotFilter::HandlePossibleStructProperty(FPropertyContainerContext& ContainerContext, FProperty* PropertyToHandle)
+{
+	if (FStructProperty* StructProperty = CastField<FStructProperty>(PropertyToHandle))
+	{
+		return FilterStructPair(ContainerContext, StructProperty);
 	}
+	return {};
+}
+
+UE::LevelSnapshots::Private::FApplySnapshotFilter::EPropertySearchResult UE::LevelSnapshots::Private::FApplySnapshotFilter::TrackChangedPropertiesInStruct(FPropertyContainerContext& ContainerContext)
+{
+	bool bFoundAtLeastOneProperty = false;
+	for (TFieldIterator<FProperty> FieldIt(ContainerContext.ContainerClass); FieldIt; ++FieldIt)
+	{
+		bFoundAtLeastOneProperty |= TrackPropertyIfChangedRecursive(ContainerContext, *FieldIt) == EPropertySearchResult::FoundProperties;
+	}
+	return bFoundAtLeastOneProperty ? EPropertySearchResult::FoundProperties : EPropertySearchResult::NoPropertiesFound;
+}
+
+void UE::LevelSnapshots::Private::FApplySnapshotFilter::PreApplyFilter(FPropertyContainerContext& ContainerContext, UObject* SnapshotObject, UObject* WorldObject)
+{
+	FLevelSnapshotsModule& Module = FLevelSnapshotsModule::GetInternalModuleInstance();
+	Module.PreApplyFilters({ SnapshotObject, WorldObject, Snapshot });
+}
+
+void UE::LevelSnapshots::Private::FApplySnapshotFilter::PostApplyFilter(FPropertyContainerContext& ContainerContext, UObject* SnapshotObject, UObject* WorldObject)
+{
+	ExtendAnalysedProperties(ContainerContext, SnapshotObject, WorldObject);
 }
 
 void UE::LevelSnapshots::Private::FApplySnapshotFilter::ExtendAnalysedProperties(FPropertyContainerContext& ContainerContext, UObject* SnapshotObject, UObject* WorldObject)
@@ -375,51 +406,32 @@ void UE::LevelSnapshots::Private::FApplySnapshotFilter::ExtendAnalysedProperties
 	}
 }
 
-TOptional<UE::LevelSnapshots::Private::FApplySnapshotFilter::EPropertySearchResult> UE::LevelSnapshots::Private::FApplySnapshotFilter::HandlePossibleStructProperties(FPropertyContainerContext& ContainerContext, FProperty* PropertyToHandle)
-{
-	if (FStructProperty* StructProperty = CastField<FStructProperty>(PropertyToHandle))
-	{
-		return FilterStructPair(ContainerContext, StructProperty);
-	}
-	return {};
-}
-
-UE::LevelSnapshots::Private::FApplySnapshotFilter::EPropertySearchResult UE::LevelSnapshots::Private::FApplySnapshotFilter::AnalyseStructProperties(FPropertyContainerContext& ContainerContext)
-{
-	bool bFoundAtLeastOneProperty = false;
-	for (TFieldIterator<FProperty> FieldIt(ContainerContext.ContainerClass); FieldIt; ++FieldIt)
-	{
-		bFoundAtLeastOneProperty |= AnalyseProperty(ContainerContext, *FieldIt) == EPropertySearchResult::FoundProperties;
-	}
-	return bFoundAtLeastOneProperty ? EPropertySearchResult::FoundProperties : EPropertySearchResult::NoPropertiesFound;
-}
-
 TOptional<UE::LevelSnapshots::Private::FApplySnapshotFilter::EPropertySearchResult> UE::LevelSnapshots::Private::FApplySnapshotFilter::HandlePossibleSubobjectProperties(FPropertyContainerContext& ContainerContext, FProperty* PropertyToHandle)
 {
 	if (FArrayProperty* ArrayProperty = CastField<FArrayProperty>(PropertyToHandle))
 	{
-		return AnalysePossibleArraySubobjectProperties(ContainerContext, ArrayProperty);
+		return TrackPossibleArraySubobjectProperties(ContainerContext, ArrayProperty);
 	}
 	if (FSetProperty* SetProperty = CastField<FSetProperty>(PropertyToHandle))
 	{
-		return AnalysePossibleSetSubobjectProperties(ContainerContext, SetProperty);
+		return TrackPossibleSetSubobjectProperties(ContainerContext, SetProperty);
 	}
 	if (FMapProperty* MapProperty = CastField<FMapProperty>(PropertyToHandle))
 	{
-		return AnalysePossibleMapSubobjectProperties(ContainerContext, MapProperty);
+		return TrackPossibleMapSubobjectProperties(ContainerContext, MapProperty);
 	}
 	
 	if (FObjectPropertyBase* ObjectProperty = CastField<FObjectPropertyBase>(PropertyToHandle))
 	{
 		void* SnapshotValuePtr = PropertyToHandle->ContainerPtrToValuePtr<void>(ContainerContext.SnapshotContainer, 0);
 		void* WorldValuePtr = PropertyToHandle->ContainerPtrToValuePtr<void>(ContainerContext.WorldContainer, 0);
-		return AnalysePossibleSubobjectProperties(ContainerContext, ObjectProperty, SnapshotValuePtr, WorldValuePtr);
+		return TrackPossibleSubobjectProperties(ContainerContext, ObjectProperty, SnapshotValuePtr, WorldValuePtr);
 	}
 	
 	return {};
 }
 
-TOptional<UE::LevelSnapshots::Private::FApplySnapshotFilter::EPropertySearchResult> UE::LevelSnapshots::Private::FApplySnapshotFilter::AnalysePossibleArraySubobjectProperties(FPropertyContainerContext& ContainerContext, FArrayProperty* PropertyToHandle)
+TOptional<UE::LevelSnapshots::Private::FApplySnapshotFilter::EPropertySearchResult> UE::LevelSnapshots::Private::FApplySnapshotFilter::TrackPossibleArraySubobjectProperties(FPropertyContainerContext& ContainerContext, FArrayProperty* PropertyToHandle)
 {
 	struct FArrayData
 	{
@@ -451,12 +463,12 @@ TOptional<UE::LevelSnapshots::Private::FApplySnapshotFilter::EPropertySearchResu
 	if (InnerAsObjectProperty && !CastField<FSoftObjectProperty>(PropertyToHandle->Inner))
 	{
 		FArrayData ArrayData(ContainerContext, PropertyToHandle);
-		return AnalysePossibleSubobjectsInCollection(ContainerContext, InnerAsObjectProperty, ArrayData);
+		return TrackPossibleSubobjectsInCollection(ContainerContext, InnerAsObjectProperty, ArrayData);
 	}
 	return {};
 }
 
-TOptional<UE::LevelSnapshots::Private::FApplySnapshotFilter::EPropertySearchResult> UE::LevelSnapshots::Private::FApplySnapshotFilter::AnalysePossibleSetSubobjectProperties(FPropertyContainerContext& ContainerContext, FSetProperty* PropertyToHandle)
+TOptional<UE::LevelSnapshots::Private::FApplySnapshotFilter::EPropertySearchResult> UE::LevelSnapshots::Private::FApplySnapshotFilter::TrackPossibleSetSubobjectProperties(FPropertyContainerContext& ContainerContext, FSetProperty* PropertyToHandle)
 {
 	struct FSetData
 	{
@@ -510,12 +522,12 @@ TOptional<UE::LevelSnapshots::Private::FApplySnapshotFilter::EPropertySearchResu
 	if (InnerAsObjectProperty && !CastField<FSoftObjectProperty>(PropertyToHandle->ElementProp))
 	{
 		FSetData SetData(ContainerContext, PropertyToHandle);
-		return AnalysePossibleSubobjectsInCollection(ContainerContext, InnerAsObjectProperty, SetData);
+		return TrackPossibleSubobjectsInCollection(ContainerContext, InnerAsObjectProperty, SetData);
 	}
 	return {};
 }
 
-TOptional<UE::LevelSnapshots::Private::FApplySnapshotFilter::EPropertySearchResult> UE::LevelSnapshots::Private::FApplySnapshotFilter::AnalysePossibleMapSubobjectProperties(FPropertyContainerContext& ContainerContext, FMapProperty* PropertyToHandle)
+TOptional<UE::LevelSnapshots::Private::FApplySnapshotFilter::EPropertySearchResult> UE::LevelSnapshots::Private::FApplySnapshotFilter::TrackPossibleMapSubobjectProperties(FPropertyContainerContext& ContainerContext, FMapProperty* PropertyToHandle)
 {
 	struct FMapData
 	{
@@ -548,13 +560,13 @@ TOptional<UE::LevelSnapshots::Private::FApplySnapshotFilter::EPropertySearchResu
 	if (InnerAsObjectProperty && !CastField<FSoftObjectProperty>(PropertyToHandle->ValueProp))
 	{
 		FMapData MapData(ContainerContext, PropertyToHandle);
-		return AnalysePossibleSubobjectsInCollection(ContainerContext, InnerAsObjectProperty, MapData);
+		return TrackPossibleSubobjectsInCollection(ContainerContext, InnerAsObjectProperty, MapData);
 	}
 	return {};
 }
 
 template<typename TCollectionData>
-TOptional<UE::LevelSnapshots::Private::FApplySnapshotFilter::EPropertySearchResult> UE::LevelSnapshots::Private::FApplySnapshotFilter::AnalysePossibleSubobjectsInCollection(FPropertyContainerContext& ContainerContext, FObjectPropertyBase* ObjectProperty, TCollectionData& Detail)
+TOptional<UE::LevelSnapshots::Private::FApplySnapshotFilter::EPropertySearchResult> UE::LevelSnapshots::Private::FApplySnapshotFilter::TrackPossibleSubobjectsInCollection(FPropertyContainerContext& ContainerContext, FObjectPropertyBase* ObjectProperty, TCollectionData& Detail)
 {
 	TArray<UObject*> UnchangedSubobjects;
 	bool bFoundChangedProperties = false;
@@ -570,7 +582,7 @@ TOptional<UE::LevelSnapshots::Private::FApplySnapshotFilter::EPropertySearchResu
 		}
 
 		void* WorldValuePtr = Detail.GetWorldValuePtrFromIndex(WorldIndex);
-		const TOptional<EPropertySearchResult> SubbjectPairResult = AnalysePossibleSubobjectProperties(ContainerContext, ObjectProperty, SnapshotValuePtr, WorldValuePtr);
+		const TOptional<EPropertySearchResult> SubbjectPairResult = TrackPossibleSubobjectProperties(ContainerContext, ObjectProperty, SnapshotValuePtr, WorldValuePtr);
 		if (!SubbjectPairResult.IsSet())
 		{
 			return {};
@@ -599,7 +611,7 @@ TOptional<UE::LevelSnapshots::Private::FApplySnapshotFilter::EPropertySearchResu
 	return bHasCollectionChanged ? EPropertySearchResult::FoundProperties : EPropertySearchResult::NoPropertiesFound;
 }
 
-TOptional<UE::LevelSnapshots::Private::FApplySnapshotFilter::EPropertySearchResult> UE::LevelSnapshots::Private::FApplySnapshotFilter::AnalysePossibleSubobjectProperties(FPropertyContainerContext& ContainerContext, FObjectPropertyBase* PropertyToHandle, void* SnapshotValuePtr, void* WorldValuePtr)
+TOptional<UE::LevelSnapshots::Private::FApplySnapshotFilter::EPropertySearchResult> UE::LevelSnapshots::Private::FApplySnapshotFilter::TrackPossibleSubobjectProperties(FPropertyContainerContext& ContainerContext, FObjectPropertyBase* PropertyToHandle, void* SnapshotValuePtr, void* WorldValuePtr)
 {
 	// Soft object properties do not work with subobjects and are not expected to: the path will remain pointing at the real world object...
 	if (FSoftObjectProperty* SoftObjectPath = CastField<FSoftObjectProperty>(PropertyToHandle))
@@ -651,7 +663,7 @@ TOptional<UE::LevelSnapshots::Private::FApplySnapshotFilter::EPropertySearchResu
 	return {};
 }
 
-UE::LevelSnapshots::Private::FApplySnapshotFilter::EPropertySearchResult UE::LevelSnapshots::Private::FApplySnapshotFilter::AnalyseProperty(FPropertyContainerContext& ContainerContext, FProperty* PropertyInCommon, bool bSkipEqualityTest)
+UE::LevelSnapshots::Private::FApplySnapshotFilter::EPropertySearchResult UE::LevelSnapshots::Private::FApplySnapshotFilter::TrackPropertyIfChangedRecursive(FPropertyContainerContext& ContainerContext, FProperty* PropertyInCommon, bool bSkipEqualityTest)
 {
 	if (!Restorability::IsRestorableProperty(PropertyInCommon))
 	{
@@ -679,7 +691,7 @@ UE::LevelSnapshots::Private::FApplySnapshotFilter::EPropertySearchResult UE::Lev
 
 UE::LevelSnapshots::Private::FApplySnapshotFilter::EPropertySearchResult UE::LevelSnapshots::Private::FApplySnapshotFilter::TrackChangedProperties(FPropertyContainerContext& ContainerContext, FProperty* PropertyInCommon)
 {
-	const TOptional<EPropertySearchResult> StructResult = HandlePossibleStructProperties(ContainerContext, PropertyInCommon);
+	const TOptional<EPropertySearchResult> StructResult = HandlePossibleStructProperty(ContainerContext, PropertyInCommon);
 	if (StructResult.IsSet())
 	{
 		return *StructResult;
@@ -691,7 +703,7 @@ UE::LevelSnapshots::Private::FApplySnapshotFilter::EPropertySearchResult UE::Lev
 		return *SubobjectResult;
 	}
 
-	if (bAllowUnchangedProperties || !AreSnapshotAndOriginalPropertiesEquivalent(Snapshot, PropertyInCommon, ContainerContext.SnapshotContainer, ContainerContext.WorldContainer, DeserializedSnapshotActor, WorldActor))
+	if (!AreSnapshotAndOriginalPropertiesEquivalent(Snapshot, PropertyInCommon, ContainerContext.SnapshotContainer, ContainerContext.WorldContainer, DeserializedSnapshotActor, WorldActor))
 	{
 		return EPropertySearchResult::FoundProperties;
 	}
diff --git a/Engine/Plugins/VirtualProduction/LevelSnapshots/Source/LevelSnapshots/Private/Filtering/Diffing/ApplySnapshotFilter.h b/Engine/Plugins/VirtualProduction/LevelSnapshots/Source/LevelSnapshots/Private/Filtering/Diffing/ApplySnapshotFilter.h
--- a/Engine/Plugins/VirtualProduction/LevelSnapshots/Source/LevelSnapshots/Private/Filtering/Diffing/ApplySnapshotFilter.h
+++ b/Engine/Plugins/VirtualProduction/LevelSnapshots/Source/LevelSnapshots/Private/Filtering/Diffing/ApplySnapshotFilter.h
@@ -17,17 +17,6 @@ namespace UE::LevelSnapshots::Private
 
 		static FApplySnapshotFilter Make(ULevelSnapshot* Snapshot, AActor* DeserializedSnapshotActor, AActor* WorldActor, const ULevelSnapshotFilter* Filter);
 		
-		FApplySnapshotFilter& AllowUnchangedProperties(bool bNewValue)
-		{
-			bAllowUnchangedProperties = bNewValue;
-			return *this;
-		}
-		FApplySnapshotFilter& AllowNonEditableProperties(bool bNewValue)
-		{
-			bAllowNonEditableProperties = bNewValue;
-			return *this;
-		}
-		
 		void ApplyFilterToFindSelectedProperties(FPropertySelectionMap& MapToAddTo);
 
 	private:
@@ -71,35 +60,36 @@ namespace UE::LevelSnapshots::Private
 		FApplySnapshotFilter(ULevelSnapshot* Snapshot, AActor* DeserializedSnapshotActor, AActor* WorldActor, const ULevelSnapshotFilter* Filter);
 		bool EnsureParametersAreValid() const;
 		
-		void AnalyseComponentProperties(FPropertySelectionMap& MapToAddTo);
+		void AnalyseComponentRestoration(FPropertySelectionMap& MapToAddTo);
 
+		// Container types: actors, subobjects, structs, custom subobjects 
 		void FilterActorPair(FPropertySelectionMap& MapToAddTo);
 		EPropertySearchResult FilterSubobjectPair(FPropertySelectionMap& MapToAddTo, UObject* SnapshotSubobject, UObject* WorldSubobject);
 		EPropertySearchResult FilterStructPair(FPropertyContainerContext& Parent, FStructProperty* StructProperty);
 		EFilterObjectPropertiesResult FindAndFilterCustomSubobjectPairs(FPropertySelectionMap& MapToAddTo, UObject* SnapshotOwner, UObject* WorldOwner);
 
-		void AnalyseRootProperties(FPropertyContainerContext& ContainerContext, UObject* SnapshotObject, UObject* WorldObject);
+		// Start property recursion
+		void TrackChangedRootProperties(FPropertyContainerContext& ContainerContext, UObject* SnapshotObject, UObject* WorldObject);
+		TOptional<EPropertySearchResult> HandlePossibleStructProperty(FPropertyContainerContext& ContainerContext, FProperty* PropertyToHandle);
+		EPropertySearchResult TrackChangedPropertiesInStruct(FPropertyContainerContext& ContainerContext);
+
+		// Registered callbacks for ILevelSnapshotsModule
+		void PreApplyFilter(FPropertyContainerContext& ContainerContext, UObject* SnapshotObject, UObject* WorldObject);
+		void PostApplyFilter(FPropertyContainerContext& ContainerContext, UObject* SnapshotObject, UObject* WorldObject);
 		void ExtendAnalysedProperties(FPropertyContainerContext& ContainerContext, UObject* SnapshotObject, UObject* WorldObject);
-		TOptional<EPropertySearchResult> HandlePossibleStructProperties(FPropertyContainerContext& ContainerContext, FProperty* PropertyToHandle);
-		EPropertySearchResult AnalyseStructProperties(FPropertyContainerContext& ContainerContext);
 
-		
+		// Subobject analysis
 		TOptional<EPropertySearchResult> HandlePossibleSubobjectProperties(FPropertyContainerContext& ContainerContext, FProperty* PropertyToHandle);
-		
-		TOptional<EPropertySearchResult> AnalysePossibleArraySubobjectProperties(FPropertyContainerContext& ContainerContext, FArrayProperty* PropertyToHandle);
-		TOptional<EPropertySearchResult> AnalysePossibleSetSubobjectProperties(FPropertyContainerContext& ContainerContext, FSetProperty* PropertyToHandle);
-		TOptional<EPropertySearchResult> AnalysePossibleMapSubobjectProperties(FPropertyContainerContext& ContainerContext, FMapProperty* PropertyToHandle);
-		
+		TOptional<EPropertySearchResult> TrackPossibleArraySubobjectProperties(FPropertyContainerContext& ContainerContext, FArrayProperty* PropertyToHandle);
+		TOptional<EPropertySearchResult> TrackPossibleSetSubobjectProperties(FPropertyContainerContext& ContainerContext, FSetProperty* PropertyToHandle);
+		TOptional<EPropertySearchResult> TrackPossibleMapSubobjectProperties(FPropertyContainerContext& ContainerContext, FMapProperty* PropertyToHandle);
 		template<typename TCollectionData>
-		TOptional<EPropertySearchResult> AnalysePossibleSubobjectsInCollection(FPropertyContainerContext& ContainerContext, FObjectPropertyBase* ObjectProperty, TCollectionData& Detail);
-		
-		TOptional<EPropertySearchResult> AnalysePossibleSubobjectProperties(FPropertyContainerContext& ContainerContext, FObjectPropertyBase* PropertyToHandle, void* SnapshotValuePtr, void* WorldValuePtr);
-
-		
+		TOptional<EPropertySearchResult> TrackPossibleSubobjectsInCollection(FPropertyContainerContext& ContainerContext, FObjectPropertyBase* ObjectProperty, TCollectionData& Detail);
+		TOptional<EPropertySearchResult> TrackPossibleSubobjectProperties(FPropertyContainerContext& ContainerContext, FObjectPropertyBase* PropertyToHandle, void* SnapshotValuePtr, void* WorldValuePtr);
 
-		EPropertySearchResult AnalyseProperty(FPropertyContainerContext& ContainerContext, FProperty* PropertyInCommon, bool bSkipEqualityTest = false);
+		// Single property analysis
+		EPropertySearchResult TrackPropertyIfChangedRecursive(FPropertyContainerContext& ContainerContext, FProperty* PropertyInCommon, bool bSkipEqualityTest = false);
 		EPropertySearchResult TrackChangedProperties(FPropertyContainerContext& ContainerContext, FProperty* PropertyInCommon);
-
 		
 		ULevelSnapshot* Snapshot;
 		AActor* DeserializedSnapshotActor;
@@ -107,12 +97,6 @@ namespace UE::LevelSnapshots::Private
 		const ULevelSnapshotFilter* Filter;
 		/** Every time an object is analysed, it is added here, so we do not analyse more than once. */
 		TMap<UObject*, EPropertySearchResult> AnalysedSnapshotObjects;
-
-		/* Do we allow properties that do not show up in the details panel? */
-		bool bAllowNonEditableProperties = false;
-
-		/* Do we allow adding properties that did not change to the selection map? */
-		bool bAllowUnchangedProperties = false;
 	};
 }
 
diff --git a/Engine/Plugins/VirtualProduction/LevelSnapshots/Source/LevelSnapshots/Private/LevelSnapshotsModule.cpp b/Engine/Plugins/VirtualProduction/LevelSnapshots/Source/LevelSnapshots/Private/LevelSnapshotsModule.cpp
--- a/Engine/Plugins/VirtualProduction/LevelSnapshots/Source/LevelSnapshots/Private/LevelSnapshotsModule.cpp
+++ b/Engine/Plugins/VirtualProduction/LevelSnapshots/Source/LevelSnapshots/Private/LevelSnapshotsModule.cpp
@@ -73,6 +73,9 @@ void UE::LevelSnapshots::Private::FLevelSnapshotsModule::StartupModule()
 		EngineTypesRestorationFence::RegisterSpecialEngineTypeSupport(*this);
 	});
 
+	OnPreTakeSnapshot().AddRaw(this, &FLevelSnapshotsModule::PreTakeSnapshot);
+	OnPostTakeSnapshot().AddRaw(this, &FLevelSnapshotsModule::PostTakeSnapshot);
+
 #if WITH_EDITOR
 	ISettingsModule& SettingsModule = FModuleManager::LoadModuleChecked<ISettingsModule>("Settings");
 	{
@@ -184,7 +187,7 @@ void UE::LevelSnapshots::Private::FLevelSnapshotsModule::UnregisterCustomObjectS
 
 void UE::LevelSnapshots::Private::FLevelSnapshotsModule::RegisterGlobalActorFilter(TSharedRef<IActorSnapshotFilter> Filter)
 {
-	GlobalFilters.AddUnique(Filter);
+	GlobalFilters.AddUnique(MoveTemp(Filter));
 }
 
 void UE::LevelSnapshots::Private::FLevelSnapshotsModule::UnregisterGlobalActorFilter(const TSharedRef<IActorSnapshotFilter>& Filter)
@@ -194,7 +197,7 @@ void UE::LevelSnapshots::Private::FLevelSnapshotsModule::UnregisterGlobalActorFi
 
 void UE::LevelSnapshots::Private::FLevelSnapshotsModule::RegisterSnapshotLoader(TSharedRef<ISnapshotLoader> Loader)
 {
-	SnapshotLoaders.AddUnique(Loader);
+	SnapshotLoaders.AddUnique(MoveTemp(Loader));
 }
 
 void UE::LevelSnapshots::Private::FLevelSnapshotsModule::UnregisterSnapshotLoader(const TSharedRef<ISnapshotLoader>& Loader)
@@ -204,7 +207,7 @@ void UE::LevelSnapshots::Private::FLevelSnapshotsModule::UnregisterSnapshotLoade
 
 void UE::LevelSnapshots::Private::FLevelSnapshotsModule::RegisterRestorationListener(TSharedRef<IRestorationListener> Listener)
 {
-	RestorationListeners.AddUnique(Listener);
+	RestorationListeners.AddUnique(MoveTemp(Listener));
 }
 
 void UE::LevelSnapshots::Private::FLevelSnapshotsModule::UnregisterRestorationListener(const TSharedRef<IRestorationListener>& Listener)
@@ -212,9 +215,19 @@ void UE::LevelSnapshots::Private::FLevelSnapshotsModule::UnregisterRestorationLi
 	RestorationListeners.RemoveSingle(Listener);
 }
 
+void UE::LevelSnapshots::Private::FLevelSnapshotsModule::RegisterTakeSnapshotListener(TSharedRef<ITakeSnapshotListener> Extender)
+{
+	TakeSnapshotListenersListeners.AddUnique(MoveTemp(Extender));
+}
+
+void UE::LevelSnapshots::Private::FLevelSnapshotsModule::UnregisterTakeSnapshotListener(const TSharedRef<ITakeSnapshotListener>& Listener)
+{
+	TakeSnapshotListenersListeners.RemoveSingle(Listener);
+}
+
 void UE::LevelSnapshots::Private::FLevelSnapshotsModule::RegisterSnapshotFilterExtender(TSharedRef<ISnapshotFilterExtender> Extender)
 {
-	FilterExtenders.AddUnique(Extender);
+	FilterExtenders.AddUnique(MoveTemp(Extender));
 }
 
 void UE::LevelSnapshots::Private::FLevelSnapshotsModule::UnregisterSnapshotFilterExtender(const TSharedRef<ISnapshotFilterExtender>& Listener)
@@ -556,6 +569,54 @@ void UE::LevelSnapshots::Private::FLevelSnapshotsModule::OnPostRemoveComponent(c
 	}
 }
 
+void UE::LevelSnapshots::Private::FLevelSnapshotsModule::OnPreTakeObjectSnapshot(const FPreTakeObjectSnapshotEvent& Params)
+{
+	SCOPED_SNAPSHOT_CORE_TRACE(TakeSnapshotListeners);
+
+	for (const TSharedRef<ITakeSnapshotListener>& Listener : TakeSnapshotListenersListeners)
+	{
+		Listener->PreTakeObjectSnapshot(Params);
+	}
+}
+
+void UE::LevelSnapshots::Private::FLevelSnapshotsModule::OnPostTakeObjectSnapshot(const FPostTakeObjectSnapshotEvent& Params)
+{
+	SCOPED_SNAPSHOT_CORE_TRACE(TakeSnapshotListeners);
+
+	for (const TSharedRef<ITakeSnapshotListener>& Listener : TakeSnapshotListenersListeners)
+	{
+		Listener->PostTakeObjectSnapshot(Params);
+	}
+}
+
+void UE::LevelSnapshots::Private::FLevelSnapshotsModule::PreTakeSnapshot(const FPreTakeSnapshotEventData& Params)
+{
+	SCOPED_SNAPSHOT_CORE_TRACE(TakeSnapshotListeners);
+
+	for (const TSharedRef<ITakeSnapshotListener>& Listener : TakeSnapshotListenersListeners)
+	{
+		Listener->PreTakeSnapshot(Params);
+	}
+}
+
+void UE::LevelSnapshots::Private::FLevelSnapshotsModule::PostTakeSnapshot(const FPostTakeSnapshotEventData& Params)
+{
+	SCOPED_SNAPSHOT_CORE_TRACE(TakeSnapshotListeners);
+
+	for (const TSharedRef<ITakeSnapshotListener>& Listener : TakeSnapshotListenersListeners)
+	{
+		Listener->PostTakeSnapshot(Params);
+	}
+}
+
+void UE::LevelSnapshots::Private::FLevelSnapshotsModule::PreApplyFilters(const FPreApplyFiltersParams& Params)
+{
+	for (const TSharedRef<ISnapshotFilterExtender>& Extender : FilterExtenders)
+	{
+		Extender->PreApplyFilters(Params);
+	}
+}
+
 UE::LevelSnapshots::FPostApplyFiltersResult UE::LevelSnapshots::Private::FLevelSnapshotsModule::PostApplyFilters(const FPostApplyFiltersParams& Params)
 {
 	FPostApplyFiltersResult Result;
diff --git a/Engine/Plugins/VirtualProduction/LevelSnapshots/Source/LevelSnapshots/Private/LevelSnapshotsModule.h b/Engine/Plugins/VirtualProduction/LevelSnapshots/Source/LevelSnapshots/Private/LevelSnapshotsModule.h
--- a/Engine/Plugins/VirtualProduction/LevelSnapshots/Source/LevelSnapshots/Private/LevelSnapshotsModule.h
+++ b/Engine/Plugins/VirtualProduction/LevelSnapshots/Source/LevelSnapshots/Private/LevelSnapshotsModule.h
@@ -42,6 +42,8 @@ namespace UE::LevelSnapshots::Private
 		virtual void UnregisterSnapshotLoader(const TSharedRef<ISnapshotLoader>& Loader) override;
 		virtual void RegisterRestorationListener(TSharedRef<IRestorationListener> Listener) override;
 		virtual void UnregisterRestorationListener(const TSharedRef<IRestorationListener>& Listener) override;
+		virtual void RegisterTakeSnapshotListener(TSharedRef<ITakeSnapshotListener> Extender) override;
+		virtual void UnregisterTakeSnapshotListener(const TSharedRef<ITakeSnapshotListener>& Listener) override;
 		virtual void RegisterSnapshotFilterExtender(TSharedRef<ISnapshotFilterExtender> Extender) override;
 		virtual void UnregisterSnapshotFilterExtender(const TSharedRef<ISnapshotFilterExtender>& Listener) override;
 		virtual void AddExplicitilySupportedProperties(const TSet<const FProperty*>& Properties) override;
@@ -95,6 +97,10 @@ namespace UE::LevelSnapshots::Private
 		void OnPreRemoveComponent(UActorComponent* ComponentToRemove);
 		void OnPostRemoveComponent(const FPostRemoveComponentParams& Params);
 
+		void OnPreTakeObjectSnapshot(const FPreTakeObjectSnapshotEvent& Params);
+		void OnPostTakeObjectSnapshot(const FPostTakeObjectSnapshotEvent& Params);
+
+		void PreApplyFilters(const FPreApplyFiltersParams& Params);
 		FPostApplyFiltersResult PostApplyFilters(const FPostApplyFiltersParams& Params);
 		
 	private:
@@ -117,6 +123,7 @@ namespace UE::LevelSnapshots::Private
 		TArray<TSharedRef<IActorSnapshotFilter>> GlobalFilters;
 		TArray<TSharedRef<ISnapshotLoader>> SnapshotLoaders;
 		TArray<TSharedRef<IRestorationListener>> RestorationListeners;
+		TArray<TSharedRef<ITakeSnapshotListener>> TakeSnapshotListenersListeners;
 		TArray<TSharedRef<ISnapshotFilterExtender>> FilterExtenders;
 
 		/* Allows these properties even when the default behaviour would exclude them. */
@@ -129,6 +136,9 @@ namespace UE::LevelSnapshots::Private
 
 		/** Map of named delegates for confirming that a level snapshot is possible. */
 		TMap<FName, FCanTakeSnapshot> CanTakeSnapshotDelegates;
+		
+		void PreTakeSnapshot(const FPreTakeSnapshotEventData& Params);
+		void PostTakeSnapshot(const FPostTakeSnapshotEventData& Params);
 	};
 }
 
diff --git a/Engine/Plugins/VirtualProduction/LevelSnapshots/Source/LevelSnapshots/Private/SpecialSupport/AssetUserDataRestoration.cpp b/Engine/Plugins/VirtualProduction/LevelSnapshots/Source/LevelSnapshots/Private/SpecialSupport/AssetUserDataRestoration.cpp
new file mode 100644
--- /dev/null
+++ b/Engine/Plugins/VirtualProduction/LevelSnapshots/Source/LevelSnapshots/Private/SpecialSupport/AssetUserDataRestoration.cpp
@@ -0,0 +1,169 @@
+// Copyright Epic Games, Inc. All Rights Reserved.
+
+#include "AssetUserDataRestoration.h"
+
+#include "Restorability/Interfaces/IRestorationListener.h"
+#include "Restorability/Interfaces/ISnapshotFilterExtender.h"
+#include "Restorability/Interfaces/ITakeSnapshotListener.h"
+
+#include "Components/ActorComponent.h"
+#include "Engine/AssetUserData.h"
+#include "Templates/Function.h"
+#include "UObject/Package.h"
+#include "UObject/ScriptInterface.h"
+#include "UObject/UObjectGlobals.h"
+#include "UObject/UnrealType.h"
+
+namespace UE::LevelSnapshots::Private::AssetUserDataRestoration
+{
+	/**
+	 * Makes sure that Level Snapshots ignores transient UAssetUserData in IInterface_AssetUserData providers, which is mainly UActorComponent.
+	 * It does this by removing and adding back transient objects before and after taking, applying and filtering a snapshot, respectively.
+	 */
+	class FAssetUserDataRestoration : public ITakeSnapshotListener, public ISnapshotFilterExtender, public IRestorationListener
+	{
+		const FArrayProperty* AssetUserDataProperty;
+	public:
+		
+		FAssetUserDataRestoration()
+			: AssetUserDataProperty(CastField<FArrayProperty>(UActorComponent::StaticClass()->FindPropertyByName(TEXT("AssetUserData"))))
+		{
+			check(AssetUserDataProperty);
+		}
+
+		virtual void PreTakeObjectSnapshot(const FPreTakeObjectSnapshotEvent& Params) override
+		{
+			RemoveTransientAssetUserData(Params.Object);
+		}
+		virtual void PostTakeObjectSnapshot(const FPostTakeObjectSnapshotEvent& Params) override
+		{
+			AddBackTransientAssetUserData(Params.Object);
+		}
+		
+		virtual void PreApplyFilters(const FPreApplyFiltersParams& Params) override
+		{
+			RemoveTransientAssetUserData(Params.EditorObject);
+			RemoveTransientAssetUserDataFromSnapshotObject(Params.SnapshotObject, Params.EditorObject);
+
+			// TODO: For backwards compatibility remove clear null references from UActorComponent::AssetUserData array property
+		}
+		virtual FPostApplyFiltersResult PostApplyFilters(const FPostApplyFiltersParams& Params) override
+		{
+			AddBackTransientAssetUserData(Params.EditorObject);
+			// Don't bother adding back transient asset user data that was removed in PreApplyFilters because it should always be ignored anyways
+			return {};
+		}
+
+		virtual void PreApplySnapshotProperties(const FApplySnapshotPropertiesParams& Params) override
+		{
+			RemoveTransientAssetUserData(Params.Object);
+		}
+		virtual void PostApplySnapshotProperties(const FApplySnapshotPropertiesParams& Params) override
+		{
+			AddBackTransientAssetUserData(Params.Object);
+		}
+
+	private:
+
+		TMap<UObject*, TArray<UAssetUserData*>> RemovedAssetUserData;
+
+		void RemoveTransientAssetUserDataFromSnapshotObject(UObject* SnapshotObject, const UObject* EditorCounterpart)
+		{
+			// Snapshot objects are slightly different... they are already in a transient package.
+			// This means that all asset user data is transient and would be removed.
+			// Hence, we'll remove those objects that match by class (since every UClass in the asset user data must be unique).
+
+			if (const TArray<UAssetUserData*>* RemovedFromEditorCounterpart = RemovedAssetUserData.Find(EditorCounterpart))
+			{
+				const TScriptInterface<IInterface_AssetUserData> AssetUserDataInterface = SnapshotObject;
+				checkf(RemovedFromEditorCounterpart, TEXT("Should be valid because EditorCounterpart must have implemented IInterface_AssetUserData as there is data for it!"));
+				for (const UAssetUserData* EditorAssetUserData : *RemovedFromEditorCounterpart)
+				{
+					AssetUserDataInterface->RemoveUserDataOfClass(EditorAssetUserData->GetClass());
+				}
+			}
+		}
+		
+		void RemoveTransientAssetUserData(UObject* Object)
+		{
+			const TScriptInterface<IInterface_AssetUserData> AssetUserDataInterface = Object;
+			if (!AssetUserDataInterface)
+			{
+				return;
+			}
+
+			TArray<UAssetUserData*> TransientObjects;
+			ForEachTransientAssetUserData(AssetUserDataInterface, [&TransientObjects](UAssetUserData* Data)
+			{
+				TransientObjects.Add(Data);
+			});
+
+			if (!TransientObjects.IsEmpty())
+			{
+				for (UAssetUserData* Data : TransientObjects)
+				{
+					AssetUserDataInterface->RemoveUserDataOfClass(Data->GetClass());
+				}
+				
+				RemovedAssetUserData.Add(Object, MoveTemp(TransientObjects));
+			}
+		}
+
+		void AddBackTransientAssetUserData(UObject* Object)
+		{
+			TArray<UAssetUserData*> TransientDataToRestore;
+			if (RemovedAssetUserData.RemoveAndCopyValue(Object, TransientDataToRestore))
+			{
+				const TScriptInterface<IInterface_AssetUserData> AssetUserDataInterface(Object);
+				checkf(AssetUserDataInterface, TEXT("Interface cast should have worked in RemoveTransientAssetUserData"))
+				for (UAssetUserData* Data : TransientDataToRestore)
+				{
+					AssetUserDataInterface->AddAssetUserData(Data);
+				}
+			}
+		}
+
+		void ForEachTransientAssetUserData(const TScriptInterface<IInterface_AssetUserData>& AssetUserDataInterface, TFunctionRef<void(UAssetUserData* TransientAssetUserData)> ProcessData)
+		{
+			auto IsTransient = [](UAssetUserData* Data)
+			{
+				UPackage* Package = Data ? Data->GetPackage() : nullptr;
+				return Data && (Data->HasAnyFlags(RF_Transient) || Package->HasAnyFlags(RF_Transient));
+			};
+
+			if (const TArray<UAssetUserData*>* Array = AssetUserDataInterface->GetAssetUserDataArray())
+			{
+				for (UAssetUserData* AssetUserData : *Array)
+				{
+					if (IsTransient(AssetUserData))
+					{
+						ProcessData(AssetUserData);
+					}
+				}
+			}
+			// UActorComponent does not implement GetAssetUserDataArray ...
+			else if (UActorComponent* Component = Cast<UActorComponent>(AssetUserDataInterface.GetObject()))
+			{
+				// ... so we must hack access to it with reflection
+				FScriptArrayHelper AssetUserDataArray(AssetUserDataProperty, AssetUserDataProperty->ContainerPtrToValuePtr<void>(Component));
+				for (int32 i = 0; i < AssetUserDataArray.Num(); ++i)
+				{
+					void* ElementPtr = AssetUserDataArray.GetRawPtr(i);
+					TObjectPtr<UAssetUserData>* TypedPtr = static_cast<TObjectPtr<UAssetUserData>*>(ElementPtr);
+					if (IsTransient(TypedPtr->Get()))
+					{
+						ProcessData(TypedPtr->Get());
+					}
+				}
+			}
+		}
+	};
+	
+	void Register(FLevelSnapshotsModule& Module)
+	{
+		const TSharedRef<FAssetUserDataRestoration> AssetUserDataRestoration = MakeShared<FAssetUserDataRestoration>();
+		Module.RegisterTakeSnapshotListener(AssetUserDataRestoration);
+		Module.RegisterSnapshotFilterExtender(AssetUserDataRestoration);
+		Module.RegisterRestorationListener(AssetUserDataRestoration);
+	}
+}
\ No newline at end of file
diff --git a/Engine/Plugins/VirtualProduction/LevelSnapshots/Source/LevelSnapshots/Private/SpecialSupport/AssetUserDataRestoration.h b/Engine/Plugins/VirtualProduction/LevelSnapshots/Source/LevelSnapshots/Private/SpecialSupport/AssetUserDataRestoration.h
new file mode 100644
--- /dev/null
+++ b/Engine/Plugins/VirtualProduction/LevelSnapshots/Source/LevelSnapshots/Private/SpecialSupport/AssetUserDataRestoration.h
@@ -0,0 +1,10 @@
+// Copyright Epic Games, Inc. All Rights Reserved.
+
+#pragma once
+
+#include "LevelSnapshotsModule.h"
+
+namespace UE::LevelSnapshots::Private::AssetUserDataRestoration
+{
+	void Register(FLevelSnapshotsModule& Module);
+}
diff --git a/Engine/Plugins/VirtualProduction/LevelSnapshots/Source/LevelSnapshots/Private/SpecialSupport/EngineTypesRestorationFence.cpp b/Engine/Plugins/VirtualProduction/LevelSnapshots/Source/LevelSnapshots/Private/SpecialSupport/EngineTypesRestorationFence.cpp
--- a/Engine/Plugins/VirtualProduction/LevelSnapshots/Source/LevelSnapshots/Private/SpecialSupport/EngineTypesRestorationFence.cpp
+++ b/Engine/Plugins/VirtualProduction/LevelSnapshots/Source/LevelSnapshots/Private/SpecialSupport/EngineTypesRestorationFence.cpp
@@ -4,6 +4,7 @@
 
 #include "ActorGroupRestoration.h"
 #include "ActorLabelRestoration.h"
+#include "AssetUserDataRestoration.h"
 #include "AttachParentShowsTransformProperties.h"
 #include "CollisionRestoration.h"
 #include "GridPlacementRestoration.h"
@@ -132,5 +133,6 @@ namespace UE::LevelSnapshots::Private::EngineTypesRestorationFence
 		MaterialInstanceSupport::Register(Module);
 		ActorLabelRestoration::Register(Module);
 		PCGRestoration::Register(Module);
+		AssetUserDataRestoration::Register(Module);
 	}
 }
diff --git a/Engine/Plugins/VirtualProduction/LevelSnapshots/Source/LevelSnapshots/Private/Tests/AssetUserDataTests.cpp b/Engine/Plugins/VirtualProduction/LevelSnapshots/Source/LevelSnapshots/Private/Tests/AssetUserDataTests.cpp
new file mode 100644
--- /dev/null
+++ b/Engine/Plugins/VirtualProduction/LevelSnapshots/Source/LevelSnapshots/Private/Tests/AssetUserDataTests.cpp
@@ -0,0 +1,261 @@
+// Copyright Epic Games, Inc. All Rights Reserved.
+
+#include "Data/LevelSnapshot.h"
+#include "Filtering/PropertySelectionMap.h"
+#include "Util/SnapshotTestRunner.h"
+#include "Types/SnapshotTestActor.h"
+#include "Types/LevelSnapshotsTestAssetUserData.h"
+
+#include "Engine/World.h"
+#include "GameFramework/Actor.h"
+#include "Misc/AutomationTest.h"
+#include "UObject/Package.h"
+
+namespace UE::LevelSnapshots::Private::Tests
+{
+	/** Shared logic for testing restoration of an IInterface_AssetUserData. */
+	class FAssetUserDataTestCase
+	{
+		FString BaseName;
+		FAutomationTestBase& Test;
+
+		FString MakeTestName(const TCHAR* TestName) const
+		{
+			return FString::Printf(TEXT("%s - %s"), *BaseName, TestName);
+		}
+		
+	public:
+
+		FAssetUserDataTestCase(FString BaseName, FAutomationTestBase& Test)
+			: BaseName(MoveTemp(BaseName))
+			, Test(Test)
+		{}
+		virtual ~FAssetUserDataTestCase() = default;
+		
+		virtual TScriptInterface<IInterface_AssetUserData> SpawnAssetUserDataStorerObject(UWorld* World) = 0;
+		virtual TScriptInterface<IInterface_AssetUserData> GetAsserUserDataStorerObjectInSnapshotWorld(ULevelSnapshot* Snapshot, const TScriptInterface<IInterface_AssetUserData>& MainWorldObject) = 0;
+
+		void RunTest()
+		{
+			ULevelSnapshotsTestAssetUserData_Persistent* Data_Persistent;
+			ULevelSnapshotsTestAssetUserData_MarkedTransient* Data_MarkedTransient;
+			ULevelSnapshotsTestAssetUserData_TransientPackage* Data_TransientPackage;
+			
+			ULevelSnapshotsTestAssetUserData_MarkedTransient* Data_MarkedTransient_Snapshot;
+			ULevelSnapshotsTestAssetUserData_TransientPackage* Data_TransientPackage_Snapshot;
+			
+			TScriptInterface<IInterface_AssetUserData> DataStorage;
+			auto GetActorFromDataStorage = [&DataStorage]()
+			{
+				return DataStorage.GetObject()->IsA<AActor>()
+					? Cast<AActor>(DataStorage.GetObject())
+					: DataStorage.GetObject()->GetTypedOuter<AActor>();
+			};
+
+			// Test must take place in a non-transient world otherwise ULevelSnapshotsTestAssetUserData_Persistent would be removed during tests since it would be in a transient package.
+			FSnapshotTestRunner(ETestFlags::NonTransientWorld)
+				.ModifyWorld([&](UWorld* World)
+				{
+					DataStorage = SpawnAssetUserDataStorerObject(World);
+					
+					Data_Persistent = NewObject<ULevelSnapshotsTestAssetUserData_Persistent>(DataStorage.GetObject());
+					Data_MarkedTransient = NewObject<ULevelSnapshotsTestAssetUserData_MarkedTransient>(DataStorage.GetObject(), ULevelSnapshotsTestAssetUserData_MarkedTransient::StaticClass(), NAME_None, RF_Transient);
+					Data_TransientPackage = NewObject<ULevelSnapshotsTestAssetUserData_TransientPackage>(GetTransientPackage());
+
+					Data_Persistent->Value = 10;
+
+					// Intentionally put the persistent data inbetween the two transient ones that will be removed to test more positioning cases.
+					DataStorage->AddAssetUserData(Data_MarkedTransient);
+					DataStorage->AddAssetUserData(Data_Persistent);
+					DataStorage->AddAssetUserData(Data_TransientPackage);
+				})
+
+				// After taking a snapshot the asset user data is still there
+				.TakeSnapshot()
+				.ModifyWorld([&](UWorld* World)
+				{
+					ULevelSnapshotsTestAssetUserData_Persistent* Retrieved_Persistent = DataStorage->GetAssetUserData<ULevelSnapshotsTestAssetUserData_Persistent>();
+					ULevelSnapshotsTestAssetUserData_MarkedTransient* Retrieved_MarkedTransient = DataStorage->GetAssetUserData<ULevelSnapshotsTestAssetUserData_MarkedTransient>();
+					ULevelSnapshotsTestAssetUserData_TransientPackage* Retrieved_TransientPackage = DataStorage->GetAssetUserData<ULevelSnapshotsTestAssetUserData_TransientPackage>();
+					
+					Test.TestTrue(MakeTestName(TEXT("Persistent - After taking snapshot")), Retrieved_Persistent == Data_Persistent);
+					Test.TestTrue(MakeTestName(TEXT("Marked Transient - After taking snapshot")), Retrieved_MarkedTransient == Data_MarkedTransient);
+					Test.TestTrue(MakeTestName(TEXT("Transient Package - After taking snapshot")), Retrieved_TransientPackage == Data_TransientPackage);
+					// Not implemented by UActorComponent...
+					if (DataStorage->GetAssetUserDataArray())
+					{
+						Test.TestEqual(MakeTestName(TEXT("No null entries - After taking snapshot")), DataStorage->GetAssetUserDataArray()->Num(), 3);
+					}
+				})
+			
+				// After filtering a snapshot the asset user data is still there
+				.FilterProperties(GetActorFromDataStorage(), [&](const FPropertySelectionMap& SelectionMap)
+				{
+					ULevelSnapshotsTestAssetUserData_Persistent* Retrieved_Persistent = DataStorage->GetAssetUserData<ULevelSnapshotsTestAssetUserData_Persistent>();
+					ULevelSnapshotsTestAssetUserData_MarkedTransient* Retrieved_MarkedTransient = DataStorage->GetAssetUserData<ULevelSnapshotsTestAssetUserData_MarkedTransient>();
+					ULevelSnapshotsTestAssetUserData_TransientPackage* Retrieved_TransientPackage = DataStorage->GetAssetUserData<ULevelSnapshotsTestAssetUserData_TransientPackage>();
+					
+					Test.TestTrue(MakeTestName(TEXT("Persistent - After filtering")), Retrieved_Persistent == Data_Persistent);
+					Test.TestTrue(MakeTestName(TEXT("Marked Transient - After filtering")), Retrieved_MarkedTransient == Data_MarkedTransient);
+					Test.TestTrue(MakeTestName(TEXT("Transient Package - After filtering")), Retrieved_TransientPackage == Data_TransientPackage);
+					// Not implemented by UActorComponent...
+					if (DataStorage->GetAssetUserDataArray())
+					{
+						Test.TestEqual(MakeTestName(TEXT("No null entries - After filtering")), DataStorage->GetAssetUserDataArray()->Num(), 3);
+					}
+
+					const TArray<FSoftObjectPath> ModifiedObjects = SelectionMap.GetKeys();
+					Test.TestEqual(MakeTestName(TEXT("No changes detected")), ModifiedObjects.Num(), 0);
+				})
+			
+				// After applying a snapshot without diff, the data is still there
+				.ApplySnapshot()
+				.RunTest([&]()
+				{
+					ULevelSnapshotsTestAssetUserData_Persistent* Retrieved_Persistent = DataStorage->GetAssetUserData<ULevelSnapshotsTestAssetUserData_Persistent>();
+					ULevelSnapshotsTestAssetUserData_MarkedTransient* Retrieved_MarkedTransient = DataStorage->GetAssetUserData<ULevelSnapshotsTestAssetUserData_MarkedTransient>();
+					ULevelSnapshotsTestAssetUserData_TransientPackage* Retrieved_TransientPackage = DataStorage->GetAssetUserData<ULevelSnapshotsTestAssetUserData_TransientPackage>();
+					
+					Test.TestTrue(MakeTestName(TEXT("Persistent - After applying")), Retrieved_Persistent == Data_Persistent);
+					Test.TestTrue(MakeTestName(TEXT("Marked Transient - After applying")), Retrieved_MarkedTransient == Data_MarkedTransient);
+					Test.TestTrue(MakeTestName(TEXT("Transient Package - After applying")), Retrieved_TransientPackage == Data_TransientPackage);
+					// Not implemented by UActorComponent...
+					if (DataStorage->GetAssetUserDataArray())
+					{
+						Test.TestEqual(MakeTestName(TEXT("No null entries - After applying")), DataStorage->GetAssetUserDataArray()->Num(), 3);
+					}
+				})
+
+				// Removing the transient data will not restore it nor add null entries ...
+				.ModifyWorld([&](UWorld* World)
+				{
+					DataStorage->RemoveUserDataOfClass(ULevelSnapshotsTestAssetUserData_MarkedTransient::StaticClass());
+					DataStorage->RemoveUserDataOfClass(ULevelSnapshotsTestAssetUserData_TransientPackage::StaticClass());
+				})
+				.FilterProperties(GetActorFromDataStorage(), [&](const FPropertySelectionMap& SelectionMap)
+				{
+					const TArray<FSoftObjectPath> ModifiedObjects = SelectionMap.GetKeys();
+					Test.TestEqual(MakeTestName(TEXT("Transient objects should not be added back")), ModifiedObjects.Num(), 0);
+				})
+				.ApplySnapshot()
+				.RunTest([&]()
+				{
+					ULevelSnapshotsTestAssetUserData_Persistent* Retrieved_Persistent = DataStorage->GetAssetUserData<ULevelSnapshotsTestAssetUserData_Persistent>();
+					Test.TestTrue(MakeTestName(TEXT("Persistent - Still exists after removing transient data")), Retrieved_Persistent == Data_Persistent);
+					// Not implemented by UActorComponent...
+					if (DataStorage->GetAssetUserDataArray())
+					{
+						Test.TestEqual(MakeTestName(TEXT("Persistent - No null entries were added after removing transient data")), DataStorage->GetAssetUserDataArray()->Num(), 1);
+					}
+				})
+			
+				// ... and adding transient data will not prompt to remove it
+				.TakeSnapshot(TEXT("WithoutTransientData"))
+				.FilterProperties(GetActorFromDataStorage(), [&](const FPropertySelectionMap& SelectionMap)
+				{
+					const TArray<FSoftObjectPath> ModifiedObjects = SelectionMap.GetKeys();
+					Test.TestEqual(MakeTestName(TEXT("Transient objects should not be added back")), ModifiedObjects.Num(), 0);
+				}, nullptr, TEXT("WithoutTransientData"))
+				.ApplySnapshot(nullptr, TEXT("WithoutTransientData"))
+				.RunTest([&]()
+				{
+					ULevelSnapshotsTestAssetUserData_Persistent* Retrieved_Persistent = DataStorage->GetAssetUserData<ULevelSnapshotsTestAssetUserData_Persistent>();
+					Test.TestTrue(MakeTestName(TEXT("Persistent - Still exists after removing transient data")), Retrieved_Persistent == Data_Persistent);
+					// Not implemented by UActorComponent...
+					if (DataStorage->GetAssetUserDataArray())
+					{
+						Test.TestEqual(MakeTestName(TEXT("Persistent - No null entries were added after removing transient data")), DataStorage->GetAssetUserDataArray()->Num(), 1);
+					}
+				})
+			
+				// Manually add back the transient data to perform more tests
+				.ModifyWorld([&](UWorld* World)
+				{
+					DataStorage->AddAssetUserData(Data_MarkedTransient);
+					DataStorage->AddAssetUserData(Data_TransientPackage);
+				})
+
+				// Simulate some engine system adding the transient object to the snapshot version...
+				.AccessSnapshot([&](ULevelSnapshot* Snapshot)
+				{
+					const TScriptInterface<IInterface_AssetUserData> SnapshotObject = GetAsserUserDataStorerObjectInSnapshotWorld(Snapshot, DataStorage);
+					Data_MarkedTransient_Snapshot = NewObject<ULevelSnapshotsTestAssetUserData_MarkedTransient>(SnapshotObject.GetObject(), ULevelSnapshotsTestAssetUserData_MarkedTransient::StaticClass(), NAME_None, RF_Transient);
+					Data_TransientPackage_Snapshot = NewObject<ULevelSnapshotsTestAssetUserData_TransientPackage>(GetTransientPackage());
+					SnapshotObject->AddAssetUserData(Data_MarkedTransient_Snapshot);
+					SnapshotObject->AddAssetUserData(Data_TransientPackage_Snapshot);
+				})
+				.FilterProperties(GetActorFromDataStorage(), [&](const FPropertySelectionMap& SelectionMap)
+				{
+					const TArray<FSoftObjectPath> ModifiedObjects = SelectionMap.GetKeys();
+					Test.TestEqual(MakeTestName(TEXT("Transient objects in snapshot world are ignored")), ModifiedObjects.Num(), 0);
+				})
+			
+				// Persistent asset user data is restored and transient ones are left untouched
+				.ModifyWorld([&](UWorld* World)
+				{
+					DataStorage->RemoveUserDataOfClass(ULevelSnapshotsTestAssetUserData_Persistent::StaticClass());
+				})
+				.ApplySnapshot()
+				.RunTest([&]()
+				{
+					ULevelSnapshotsTestAssetUserData_Persistent* Retrieved_Persistent = DataStorage->GetAssetUserData<ULevelSnapshotsTestAssetUserData_Persistent>();
+					ULevelSnapshotsTestAssetUserData_MarkedTransient* Retrieved_MarkedTransient = DataStorage->GetAssetUserData<ULevelSnapshotsTestAssetUserData_MarkedTransient>();
+					ULevelSnapshotsTestAssetUserData_TransientPackage* Retrieved_TransientPackage = DataStorage->GetAssetUserData<ULevelSnapshotsTestAssetUserData_TransientPackage>();
+					
+					Test.TestTrue(MakeTestName(TEXT("Persistent - After removing and applying")), Retrieved_Persistent != nullptr);
+					if (Retrieved_Persistent)
+					{
+						Test.TestEqual(MakeTestName(TEXT("Persistent - Correct value restored")), Retrieved_Persistent->Value, 10);
+					}
+					
+					Test.TestTrue(MakeTestName(TEXT("Marked Transient - After removing and applying")), Retrieved_MarkedTransient == Data_MarkedTransient);
+					Test.TestTrue(MakeTestName(TEXT("Transient Package - After removing and applying")), Retrieved_TransientPackage == Data_TransientPackage);
+					// Not implemented by UActorComponent...
+					if (DataStorage->GetAssetUserDataArray())
+					{
+						Test.TestEqual(MakeTestName(TEXT("No null entries - After removing and applying")), DataStorage->GetAssetUserDataArray()->Num(), 3);
+					}
+				})
+			;
+		}
+	};
+
+	/** Stores asset in an actor's component. */
+	class FAssetUserDataTestCase_InComponent : public FAssetUserDataTestCase
+	{
+	public:
+
+		using FAssetUserDataTestCase::FAssetUserDataTestCase;
+
+		virtual TScriptInterface<IInterface_AssetUserData> SpawnAssetUserDataStorerObject(UWorld* World) override
+		{
+			return ASnapshotTestActor::Spawn(World)->TestComponent;
+		}
+
+		virtual TScriptInterface<IInterface_AssetUserData> GetAsserUserDataStorerObjectInSnapshotWorld(ULevelSnapshot* Snapshot, const TScriptInterface<IInterface_AssetUserData>& MainWorldObject)
+		{
+			UActorComponent* Component = Cast<UActorComponent>(MainWorldObject.GetObject());
+			TOptional<TNonNullPtr<AActor>> SnapshotActor = Snapshot->GetDeserializedActor(Component->GetOwner());
+			return SnapshotActor ? Cast<ASnapshotTestActor>(SnapshotActor.GetValue())->TestComponent : nullptr;
+		}
+	};
+	
+	/**
+	 * This tests that UActorComponent::AssetUserData is restored correctly. 
+	 * Some engine systems add transient objects to the AssetUserData. Under normal LS conditions, such objects would be nulled when restoring.
+	 * 
+	 * Level Snapshots has no general system for handling transient instanced objects in UPROPERTIES but there are special rules for AssetUserData:
+	 * - If an AssetUserData object is transient and we're taking a snapshot, the object is removed before and added back after taking the snapshot.
+	 * - If an AssetUserData object is transient and we're diffing the world, the object is removed before and added back after diffing the snapshot.
+	 */
+	IMPLEMENT_SIMPLE_AUTOMATION_TEST(FAssetUserDataTests, "VirtualProduction.LevelSnapshots.Snapshot.Other.AssetUserData", (EAutomationTestFlags::ApplicationContextMask | EAutomationTestFlags::EngineFilter));
+	bool FAssetUserDataTests::RunTest(const FString& Parameters)
+	{
+		FAssetUserDataTestCase_InComponent ComponentTest(TEXT("ComponentTest"), *this);
+		// Adding a test for world settings was considered but AWorldSettings::AssetUserData property is not CPF_Editable so Level Snapshots does not restore it
+
+		ComponentTest.RunTest();
+		
+		return true;
+	}
+}
diff --git a/Engine/Plugins/VirtualProduction/LevelSnapshots/Source/LevelSnapshots/Private/Tests/Types/LevelSnapshotsTestAssetUserData.h b/Engine/Plugins/VirtualProduction/LevelSnapshots/Source/LevelSnapshots/Private/Tests/Types/LevelSnapshotsTestAssetUserData.h
new file mode 100644
--- /dev/null
+++ b/Engine/Plugins/VirtualProduction/LevelSnapshots/Source/LevelSnapshots/Private/Tests/Types/LevelSnapshotsTestAssetUserData.h
@@ -0,0 +1,38 @@
+// Copyright Epic Games, Inc. All Rights Reserved.
+
+#pragma once
+
+#include "Engine/AssetUserData.h"
+#include "LevelSnapshotsTestAssetUserData.generated.h"
+
+// These classes do nothing special. See UActorComponent::AddAssetUserData: it only allows one instance per unique class... this is why we must create a few additional classes.
+
+UCLASS()
+class LEVELSNAPSHOTS_API ULevelSnapshotsTestAssetUserData_Persistent : public UAssetUserData
+{
+	GENERATED_BODY()
+public:
+
+	UPROPERTY(EditAnywhere, Category = "Level Snapshots")
+	int32 Value = 0;
+};
+
+UCLASS()
+class LEVELSNAPSHOTS_API ULevelSnapshotsTestAssetUserData_MarkedTransient : public UAssetUserData
+{
+	GENERATED_BODY()
+public:
+
+	UPROPERTY(EditAnywhere, Category = "Level Snapshots")
+	int32 Value = 0;
+};
+
+UCLASS()
+class LEVELSNAPSHOTS_API ULevelSnapshotsTestAssetUserData_TransientPackage : public UAssetUserData
+{
+	GENERATED_BODY()
+public:
+
+	UPROPERTY(EditAnywhere, Category = "Level Snapshots")
+	int32 Value = 0;
+};
diff --git a/Engine/Plugins/VirtualProduction/LevelSnapshots/Source/LevelSnapshots/Private/Tests/Util/SnapshotTestRunner.cpp b/Engine/Plugins/VirtualProduction/LevelSnapshots/Source/LevelSnapshots/Private/Tests/Util/SnapshotTestRunner.cpp
--- a/Engine/Plugins/VirtualProduction/LevelSnapshots/Source/LevelSnapshots/Private/Tests/Util/SnapshotTestRunner.cpp
+++ b/Engine/Plugins/VirtualProduction/LevelSnapshots/Source/LevelSnapshots/Private/Tests/Util/SnapshotTestRunner.cpp
@@ -7,30 +7,59 @@
 #include "LevelSnapshotsFunctionLibrary.h"
 
 #include "PreviewScene.h"
+#include "UObject/Package.h"
 
 FName UE::LevelSnapshots::Private::Tests::FSnapshotTestRunner::DefaultSnapshotId = FName("DefaultSnapshotId");
 
-UE::LevelSnapshots::Private::Tests::FSnapshotTestRunner::FSnapshotTestRunner()
+UE::LevelSnapshots::Private::Tests::FSnapshotTestRunner::FSnapshotTestRunner(ETestFlags Flags)
+	: Flags(Flags)
 {
 	TestWorld = MakeShared<FPreviewScene>(
 		FPreviewScene::ConstructionValues()
 			.SetEditor(true)
 			);
+	
+	// Some tests rely on the fact that the test takes place in a non-transient world
+	if (EnumHasAnyFlags(Flags, ETestFlags::NonTransientWorld))
+	{
+		UPackage* Package = CreatePackage(*FString::Printf(TEXT("/Temp/LevelSnapshots/Test%s"), *FGuid::NewGuid().ToString(EGuidFormats::Base36Encoded)));
+		// Executed tests might (indirectly) trigger a manual garbage collection 
+		Package->AddToRoot();
+		TestWorld->GetWorld()->Rename(nullptr, Package);
+	}
 }
 
 UE::LevelSnapshots::Private::Tests::FSnapshotTestRunner::FSnapshotTestRunner(UE::LevelSnapshots::Private::Tests::FSnapshotTestRunner&& Other)
 {
 	TestWorld = Other.TestWorld;
 	Snapshots = Other.Snapshots;
+	Flags = Other.Flags;
 	// Prevent Other's destructor from calling RemoveFromRoot
 	Other.Snapshots.Empty();
 }
 
 UE::LevelSnapshots::Private::Tests::FSnapshotTestRunner::~FSnapshotTestRunner()
 {
+	if (EnumHasAnyFlags(Flags, ETestFlags::NonTransientWorld))
+	{
+		UPackage* Package = TestWorld->GetWorld()->GetPackage();
+		Package->RemoveFromRoot();
+		// If you run a test, and then try to close the engine, the package would show as pending save
+		Package->SetFlags(RF_Transient);
+	}
+	
 	for (auto SnapshotIt = Snapshots.CreateIterator(); SnapshotIt; ++SnapshotIt)
 	{
+		SnapshotIt->Value->ConditionalBeginDestroy();
 		SnapshotIt->Value->RemoveFromRoot();
+		// ULevelSnapshot::BeginDestroy marks it as RF_Standalone again
+		SnapshotIt->Value->ClearFlags(RF_Standalone);
+	}
+	
+	if (EnumHasAnyFlags(Flags, ETestFlags::NonTransientWorld))
+	{
+		// Otherwise other tests will log a warning "CleanupWorld was called twice"
+		CollectGarbage(RF_Standalone);
 	}
 }
 
@@ -125,7 +154,8 @@ UE::LevelSnapshots::Private::Tests::FSnapshotTestRunner& UE::LevelSnapshots::Pri
 {
 	return AccessSnapshot([OriginalActor, &Callback, Filter](ULevelSnapshot* Snapshot)
 	{
-		const TOptional<TNonNullPtr<AActor>> SnapshotCounterpart = Snapshot->GetDeserializedActor(OriginalActor);
+		const FSoftObjectPath PathToOriginalActor(OriginalActor);
+		const TOptional<TNonNullPtr<AActor>> SnapshotCounterpart = Snapshot->GetDeserializedActor(PathToOriginalActor);
 		if (ensure(SnapshotCounterpart))
 		{
 			FPropertySelectionMap SelectedProperties;
diff --git a/Engine/Plugins/VirtualProduction/LevelSnapshots/Source/LevelSnapshots/Private/Tests/Util/SnapshotTestRunner.h b/Engine/Plugins/VirtualProduction/LevelSnapshots/Source/LevelSnapshots/Private/Tests/Util/SnapshotTestRunner.h
--- a/Engine/Plugins/VirtualProduction/LevelSnapshots/Source/LevelSnapshots/Private/Tests/Util/SnapshotTestRunner.h
+++ b/Engine/Plugins/VirtualProduction/LevelSnapshots/Source/LevelSnapshots/Private/Tests/Util/SnapshotTestRunner.h
@@ -13,6 +13,14 @@ class FPreviewScene;
 
 namespace UE::LevelSnapshots::Private::Tests
 {
+	enum class ETestFlags
+	{
+		None,
+		/** The UWorld should be created in a non-transient package */
+		NonTransientWorld = 1 << 0
+	};
+	ENUM_CLASS_FLAGS(ETestFlags);
+	
 	/**
 	* Utility for executing tests
 	*/
@@ -22,7 +30,7 @@ namespace UE::LevelSnapshots::Private::Tests
 
 		static FName DefaultSnapshotId;
 
-		FSnapshotTestRunner();
+		FSnapshotTestRunner(ETestFlags Flags = ETestFlags::None);
 		FSnapshotTestRunner(const FSnapshotTestRunner&) = delete;
 		FSnapshotTestRunner(FSnapshotTestRunner&& Other);
 		~FSnapshotTestRunner();
@@ -46,6 +54,8 @@ namespace UE::LevelSnapshots::Private::Tests
 
 	private:
 
+		ETestFlags Flags;
+
 		TSharedPtr<FPreviewScene> TestWorld;
 
 		TMap<FName, ULevelSnapshot*> Snapshots;
diff --git a/Engine/Plugins/VirtualProduction/LevelSnapshots/Source/LevelSnapshots/Public/ILevelSnapshotsModule.h b/Engine/Plugins/VirtualProduction/LevelSnapshots/Source/LevelSnapshots/Public/ILevelSnapshotsModule.h
--- a/Engine/Plugins/VirtualProduction/LevelSnapshots/Source/LevelSnapshots/Public/ILevelSnapshotsModule.h
+++ b/Engine/Plugins/VirtualProduction/LevelSnapshots/Source/LevelSnapshots/Public/ILevelSnapshotsModule.h
@@ -2,7 +2,6 @@
 
 #pragma once
 
-#include "CoreMinimal.h"
 #include "Interfaces/IRestorationListener.h"
 #include "Interfaces/ISnapshotLoader.h"
 #include "Interfaces/IPropertyComparer.h"
@@ -10,6 +9,7 @@
 #include "Interfaces/ICustomObjectSnapshotSerializer.h"
 #include "Interfaces/IActorSnapshotFilter.h"
 #include "Interfaces/ISnapshotFilterExtender.h"
+#include "Interfaces/ITakeSnapshotListener.h"
 
 #include "Modules/ModuleInterface.h"
 #include "Modules/ModuleManager.h"
@@ -21,11 +21,13 @@ namespace UE::LevelSnapshots
 	struct LEVELSNAPSHOTS_API FPreTakeSnapshotEventData
 	{
 		ULevelSnapshot* Snapshot;
+		UWorld* World;
 	};
 
 	struct LEVELSNAPSHOTS_API FPostTakeSnapshotEventData
 	{
 		ULevelSnapshot* Snapshot;
+		UWorld* World;
 	};
 
 	class LEVELSNAPSHOTS_API ILevelSnapshotsModule : public IModuleInterface
@@ -105,11 +107,16 @@ namespace UE::LevelSnapshots
 		virtual void RegisterRestorationListener(TSharedRef<IRestorationListener> Listener) = 0;
 		virtual void UnregisterRestorationListener(const TSharedRef<IRestorationListener>& Listener) = 0;
 
+		
+		/** Registers an object that listens for when objects are snapshot. */
+		virtual void RegisterTakeSnapshotListener(TSharedRef<ITakeSnapshotListener> Extender) = 0;
+		virtual void UnregisterTakeSnapshotListener(const TSharedRef<ITakeSnapshotListener>& Listener) = 0;
+		
 
 		/** Registers an object that can decide to display additional properties to the user that default snapshot behaviour does not display. */
 		virtual void RegisterSnapshotFilterExtender(TSharedRef<ISnapshotFilterExtender> Extender) = 0;
 		virtual void UnregisterSnapshotFilterExtender(const TSharedRef<ISnapshotFilterExtender>& Listener) = 0;
-		
+
 		
 		/**
 		 * Adds properties that snapshots will capture and restore from now on. This allows support for properties that are skipped by default.
diff --git a/Engine/Plugins/VirtualProduction/LevelSnapshots/Source/LevelSnapshots/Public/Restorability/Interfaces/ISnapshotFilterExtender.h b/Engine/Plugins/VirtualProduction/LevelSnapshots/Source/LevelSnapshots/Public/Restorability/Interfaces/ISnapshotFilterExtender.h
--- a/Engine/Plugins/VirtualProduction/LevelSnapshots/Source/LevelSnapshots/Public/Restorability/Interfaces/ISnapshotFilterExtender.h
+++ b/Engine/Plugins/VirtualProduction/LevelSnapshots/Source/LevelSnapshots/Public/Restorability/Interfaces/ISnapshotFilterExtender.h
@@ -10,10 +10,22 @@ struct FPropertySelectionMap;
 
 namespace UE::LevelSnapshots
 {
+	struct FPreApplyFiltersParams
+	{
+		/** The editor object being filtered */
+		UObject* SnapshotObject;
+		/** The snapshot world object being filtered */
+		UObject* EditorObject;
+		/** The snapshot on which the filtering occurs */
+		ULevelSnapshot* Snapshot;
+	};
+	
 	struct FPostApplyFiltersParams
 	{
 		const FPropertySelection& SelectedProperties;
+		/** The editor object being filtered */
 		UObject* SnapshotObject;
+		/** The snapshot world object being filtered */
 		UObject* EditorObject;
 
 		FPostApplyFiltersParams(const FPropertySelection& SelectedProperties, UObject* SnapshotObject, UObject* EditorObject)
@@ -44,10 +56,12 @@ namespace UE::LevelSnapshots
 	};
 	
 	/**
-	 * Allows you to display additional properties in the results view. This is used whenever a selection set is built.
+	 * Allows you to receive callbacks when an object is filtered and to display additional properties in the results view.
+	 * This is used whenever a selection set is built.
 	 *
 	 * By default only modified properties are shown. However, sometimes you may want to include unmodified properties, too.
-	 * One example is:
+	 * 
+	 * Example 1: Transforms and parenting
 	 *	1. Create cube
 	 *	2. Rotate it to 40 degrees
 	 *  3. Take a snapshot
@@ -60,13 +74,29 @@ namespace UE::LevelSnapshots
 	 *	actors are removed first and then the properties of modified actors are saved.
 	 *	The solution is, using this interface, to add the cube's transform properties even though they're equal to the snapshot version.
 	 *	Users will see the property in the UI and can opt to not restore it.
+	 *
+	 * Example 2: AMediaPlate
+	 * Media plates add a transient UAssetUserData object to a component.
+	 * In OnPreSaveWorld and OnPostSaveWorld the UAssetUserData would be removed and added back. Similar behaviour should
+	 * occur for Level Snapshots when filtering and applying (events about applying are exposed by IRestorationListener).
+	 * 
+	 * Note that Level Snapshots restores TArrays either fully or not at all.
+	 * For this reason, IPropertyComparer may not always be an option for you; this interface gives you more options.
+	 * If you have changes that should be ignored by Level Snapshots, this interface allows you to modify the object in PreFilterObject
+	 * and undo the changes in PostFilterObject.
 	 */
 	class LEVELSNAPSHOTS_API ISnapshotFilterExtender
 	{
 	public:
 
+		/**
+		 * Called before the object is filtered.
+		 * PreFilterObject is be executed recursively on subobjects before the corresponding call to PostFilterObject is made.
+		 */
+		virtual void PreApplyFilters(const FPreApplyFiltersParams& Params) {}
+
 		/** Called after all properties have been filtered on the given object pair. */
-		virtual FPostApplyFiltersResult PostApplyFilters(const FPostApplyFiltersParams& Params) = 0;
+		virtual FPostApplyFiltersResult PostApplyFilters(const FPostApplyFiltersParams& Params) { return {}; }
 		
 		virtual ~ISnapshotFilterExtender() = default;
 	};
diff --git a/Engine/Plugins/VirtualProduction/LevelSnapshots/Source/LevelSnapshots/Public/Restorability/Interfaces/ITakeSnapshotListener.h b/Engine/Plugins/VirtualProduction/LevelSnapshots/Source/LevelSnapshots/Public/Restorability/Interfaces/ITakeSnapshotListener.h
new file mode 100644
--- /dev/null
+++ b/Engine/Plugins/VirtualProduction/LevelSnapshots/Source/LevelSnapshots/Public/Restorability/Interfaces/ITakeSnapshotListener.h
@@ -0,0 +1,39 @@
+// Copyright Epic Games, Inc. All Rights Reserved.
+
+#pragma once
+
+class ULevelSnapshot;
+
+namespace UE::LevelSnapshots
+{
+	struct FPostTakeSnapshotEventData;
+	struct FPreTakeSnapshotEventData;
+
+	struct FTakeObjectSnapshotEvent
+	{
+		UObject* Object;
+	};
+
+	struct FPreTakeObjectSnapshotEvent : FTakeObjectSnapshotEvent
+	{};
+	struct FPostTakeObjectSnapshotEvent : FTakeObjectSnapshotEvent
+	{};
+
+	/** Exposes callbacks for when a snapshot is taken */
+	class LEVELSNAPSHOTS_API ITakeSnapshotListener
+	{
+	public:
+
+		/** Called before an object's data is recorded by a snapshot. */
+		virtual void PreTakeObjectSnapshot(const FPreTakeObjectSnapshotEvent& Params) {}
+		/** Called after an object's data is recorded by a snapshot. */
+		virtual void PostTakeObjectSnapshot(const FPostTakeObjectSnapshotEvent& Params) {}
+
+		/** Called before a snapshot is taken for a world. */
+		virtual void PreTakeSnapshot(const FPreTakeSnapshotEventData& Params) {}
+		/** Called after a snapshot is taken for a world. */
+		virtual void PostTakeSnapshot(const FPostTakeSnapshotEventData& Params) {}
+		
+		virtual ~ITakeSnapshotListener() = default;
+	};
+}
