Level Snapshots:

-Fix: Crash after restoring Media Plate properties in Level Snapshots when in Multi-User
-Fix: Level Snapshots > Crash in results editor UI when expanding VCam actor

diff --git a/Engine/Plugins/VirtualProduction/LevelSnapshots/Source/LevelSnapshotsEditor/Private/Views/Results/LevelSnapshotsEditorResultsHelpers.cpp b/Engine/Plugins/VirtualProduction/LevelSnapshots/Source/LevelSnapshotsEditor/Private/Views/Results/LevelSnapshotsEditorResultsHelpers.cpp
--- a/Engine/Plugins/VirtualProduction/LevelSnapshots/Source/LevelSnapshotsEditor/Private/Views/Results/LevelSnapshotsEditorResultsHelpers.cpp
+++ b/Engine/Plugins/VirtualProduction/LevelSnapshots/Source/LevelSnapshotsEditor/Private/Views/Results/LevelSnapshotsEditorResultsHelpers.cpp
@@ -78,7 +78,7 @@ void FLevelSnapshotsEditorResultsHelpers::LoopOverHandleHierarchiesAndCreateRowH
 	
 	const TSharedPtr<FPropertyHandleHierarchy> PinnedInputHierarchy = InputHierarchy.Pin();
 
-	const TSharedPtr<IPropertyHandle>& InputHandle = PinnedInputHierarchy->Handle;
+	const TSharedPtr<IPropertyHandle> InputHandle = PinnedInputHierarchy->Handle;
 
 	// No asserts for handle or property because PropertyRowGenerator assumes there will be a details layout pointer but this scenario doesn't create one
 	// Asserting would cause needless debugging breaks every time a snapshot is chosen
@@ -205,22 +205,32 @@ void FLevelSnapshotsEditorResultsHelpers::LoopOverHandleHierarchiesAndCreateRowH
 			UObject* WorldSubobject = nullptr;
 			UObject* SnapshotSubobject = nullptr;
 			InputHandle->GetValue(bInputIsWorld ? WorldSubobject : SnapshotSubobject);
-			CounterpartHandle->GetValue(bInputIsWorld ? SnapshotSubobject : WorldSubobject);
 
-			if (!ensure(WorldSubobject || SnapshotSubobject))
+			if (CounterpartHandle.IsValid())
+			{
+				CounterpartHandle->GetValue(bInputIsWorld ? SnapshotSubobject : WorldSubobject);
+			}
+
+			if (!ensure(WorldSubobject))
 			{
 				return;
 			}
 
-			FText ObjectName = FText::FromString((WorldSubobject ? WorldSubobject : SnapshotSubobject)->GetName());
+			const FPropertySelection* PropertySelection = PropertySelectionMap.GetObjectSelection(WorldSubobject).GetPropertySelection();
+
+			if (!PropertySelection)
+			{
+				return;
+			}
 
 			TWeakPtr<FLevelSnapshotsEditorResultsRow> NewObjectRow =
 				BuildModifiedObjectRow(
 					WorldSubobject, SnapshotSubobject, PropertyEditorModule,
 					PropertySelectionMap,
-					PropertySelectionMap.GetObjectSelection(WorldSubobject).GetPropertySelection()->GetSelectedLeafProperties(),
+					PropertySelection->GetSelectedLeafProperties(),
 					InDirectParentRow, InResultsView,
-					FText::Format(INVTEXT("{PropertyName} ({ObjectName})"), DisplayName, ObjectName));
+					FText::Format(INVTEXT("{PropertyName} ({ObjectName})"),
+						DisplayName, FText::FromString(WorldSubobject->GetName())));
 
 			if (NewObjectRow.IsValid())
 			{
