Level Snapshots:

-Fix: Crash when Expanding Media Plates modifications if parameters "AutoPlay, Enable Audio, Play Only when Visible" are changed

diff --git a/Engine/Plugins/VirtualProduction/LevelSnapshots/Source/LevelSnapshotsEditor/LevelSnapshotsEditor.build.cs b/Engine/Plugins/VirtualProduction/LevelSnapshots/Source/LevelSnapshotsEditor/LevelSnapshotsEditor.build.cs
--- a/Engine/Plugins/VirtualProduction/LevelSnapshots/Source/LevelSnapshotsEditor/LevelSnapshotsEditor.build.cs
+++ b/Engine/Plugins/VirtualProduction/LevelSnapshots/Source/LevelSnapshotsEditor/LevelSnapshotsEditor.build.cs
@@ -23,6 +23,7 @@ public class LevelSnapshotsEditor : ModuleRules
 				"ApplicationCore",
 				"AssetRegistry",
 				"AssetTools",
+				"BlueprintGraph",
 				"CoreUObject",
 				"ContentBrowser",
 				"Engine",
diff --git a/Engine/Plugins/VirtualProduction/LevelSnapshots/Source/LevelSnapshotsEditor/Private/Views/Results/LevelSnapshotsEditorResultsHelpers.cpp b/Engine/Plugins/VirtualProduction/LevelSnapshots/Source/LevelSnapshotsEditor/Private/Views/Results/LevelSnapshotsEditorResultsHelpers.cpp
--- a/Engine/Plugins/VirtualProduction/LevelSnapshots/Source/LevelSnapshotsEditor/Private/Views/Results/LevelSnapshotsEditorResultsHelpers.cpp
+++ b/Engine/Plugins/VirtualProduction/LevelSnapshots/Source/LevelSnapshotsEditor/Private/Views/Results/LevelSnapshotsEditorResultsHelpers.cpp
@@ -7,9 +7,11 @@
 #include "SLevelSnapshotsEditorResults.h"
 
 #include "Algo/Find.h"
+#include "BlueprintEditorSettings.h"
 #include "IDetailPropertyRow.h"
 #include "IDetailTreeNode.h"
 #include "IPropertyRowGenerator.h"
+#include "Kismet2/ComponentEditorUtils.h"
 #include "Modules/ModuleManager.h"
 #include "PropertyEditorModule.h"
 
@@ -517,6 +519,54 @@ UObject* FLevelSnapshotsEditorResultsHelpers::FindCounterpartComponent(const UAc
 	}
 }
 
+bool FLevelSnapshotsEditorResultsHelpers::IsInstancedActorComponentDesirableForComponentArray(
+	UActorComponent* ActorComponent, const bool bHideConstructionScriptComponentsInDetailsView)
+{
+	const USceneComponent* AsSceneComponent = Cast<USceneComponent>(ActorComponent);
+	const USceneComponent* ParentSceneComponent = AsSceneComponent != nullptr ? AsSceneComponent->GetAttachParent() : nullptr;
+
+	// Exclude nested DSOs attached to BP-constructed instances, which are not mutable.
+	return (ActorComponent != nullptr
+		&& (!ActorComponent->IsVisualizationComponent())
+		&& (ActorComponent->CreationMethod != EComponentCreationMethod::UserConstructionScript || !bHideConstructionScriptComponentsInDetailsView)
+		&& (!ParentSceneComponent || !ParentSceneComponent->IsCreatedByConstructionScript() || !ActorComponent->HasAnyFlags(RF_DefaultSubObject)))
+		&& (ActorComponent->CreationMethod != EComponentCreationMethod::Native || FComponentEditorUtils::GetPropertyForEditableNativeComponent(ActorComponent));
+}
+
+TSet<UActorComponent*> FLevelSnapshotsEditorResultsHelpers::FilterOutActorComponentsBasedOnEditorSettings(
+	const TSet<UActorComponent*>& UnfilteredActorComponentList)
+{
+	const bool bHideConstructionScriptComponentsInDetailsView =
+		GetDefault<UBlueprintEditorSettings>()->bHideConstructionScriptComponentsInDetailsView;
+
+	// Filter the components by their visibility
+	TSet<UActorComponent*> ComponentsToAdd;
+	for (UActorComponent* ActorComp : UnfilteredActorComponentList)
+	{
+		if (ActorComp && IsInstancedActorComponentDesirableForComponentArray(ActorComp, bHideConstructionScriptComponentsInDetailsView))
+		{
+			ComponentsToAdd.Add(ActorComp);
+		}
+	}
+
+	return ComponentsToAdd;
+}
+
+TSet<UActorComponent*> FLevelSnapshotsEditorResultsHelpers::FilterOutActorComponentsBasedOnEditorSettings(
+	const TSet<TWeakObjectPtr<UActorComponent>>& UnfilteredActorComponentList)
+{
+	TSet<UActorComponent*> RawSet;
+	for(const TWeakObjectPtr<UActorComponent>& Component : UnfilteredActorComponentList)
+	{
+		if (UActorComponent* RawComponent = Component.Get())
+		{
+			RawSet.Add(RawComponent);
+		}
+	}
+
+	return FilterOutActorComponentsBasedOnEditorSettings(RawSet);
+}
+
 int32 FLevelSnapshotsEditorResultsHelpers::CreateNewHierarchyStructInLoop(const AActor* InActor, USceneComponent* SceneComponent, TArray<TWeakPtr<FComponentHierarchy>>& AllHierarchies)
 {
 	check(InActor);
@@ -533,7 +583,7 @@ int32 FLevelSnapshotsEditorResultsHelpers::CreateNewHierarchyStructInLoop(const
 		int32 IndexOfParentHierarchy = AllHierarchies.IndexOfByPredicate(
 			[&ParentComponent](const TWeakPtr<FComponentHierarchy>& Hierarchy)
 			{
-				return Hierarchy.Pin()->Component == ParentComponent;
+				return Hierarchy.IsValid() && Hierarchy.Pin()->Component == ParentComponent;
 			});
 
 		if (IndexOfParentHierarchy == -1)
@@ -541,7 +591,14 @@ int32 FLevelSnapshotsEditorResultsHelpers::CreateNewHierarchyStructInLoop(const
 			IndexOfParentHierarchy = CreateNewHierarchyStructInLoop(InActor, ParentComponent, AllHierarchies);
 		}
 
-		AllHierarchies[IndexOfParentHierarchy].Pin()->DirectChildren.Add(NewHierarchy);
+		const bool bIsParentHierarchyValid = ensureMsgf(AllHierarchies[IndexOfParentHierarchy].IsValid(),
+			TEXT("%hs: AllHierarchies at index %d for Component '%s' was not valid."),
+			__FUNCTION__, IndexOfParentHierarchy, *(ParentComponent->GetName()));
+
+		if (bIsParentHierarchyValid)
+		{
+			AllHierarchies[IndexOfParentHierarchy].Pin()->DirectChildren.Add(NewHierarchy);
+		}
 	}
 
 	return ReturnValue;
@@ -557,11 +614,12 @@ TSharedRef<FComponentHierarchy> FLevelSnapshotsEditorResultsHelpers::BuildCompon
 	TArray<TWeakPtr<FComponentHierarchy>> AllHierarchies;
 	AllHierarchies.Add(ReturnHierarchy);
 
-	TSet<UActorComponent*> AllActorComponents = InActor->GetComponents();
+	TSet<UActorComponent*> AllActorComponents = FilterOutActorComponentsBasedOnEditorSettings(InActor->GetComponents());
 
 	for (UActorComponent* Component : AllActorComponents)
 	{
-		if (USceneComponent* SceneComponent = Cast<USceneComponent>(Component))
+		USceneComponent* SceneComponent = Cast<USceneComponent>(Component);
+		if (IsValid(SceneComponent) && !SceneComponent->HasAnyFlags(RF_Transient))
 		{
 			const bool ComponentContained = AllHierarchies.ContainsByPredicate(
 				[&SceneComponent](const TWeakPtr<FComponentHierarchy>& Hierarchy)
diff --git a/Engine/Plugins/VirtualProduction/LevelSnapshots/Source/LevelSnapshotsEditor/Private/Views/Results/LevelSnapshotsEditorResultsHelpers.h b/Engine/Plugins/VirtualProduction/LevelSnapshots/Source/LevelSnapshotsEditor/Private/Views/Results/LevelSnapshotsEditorResultsHelpers.h
--- a/Engine/Plugins/VirtualProduction/LevelSnapshots/Source/LevelSnapshotsEditor/Private/Views/Results/LevelSnapshotsEditorResultsHelpers.h
+++ b/Engine/Plugins/VirtualProduction/LevelSnapshots/Source/LevelSnapshotsEditor/Private/Views/Results/LevelSnapshotsEditorResultsHelpers.h
@@ -39,6 +39,30 @@ namespace FLevelSnapshotsEditorResultsHelpers
 
 	UObject* FindCounterpartComponent(const UActorComponent* SubObjectToMatch, const TSet<UActorComponent*>& InCounterpartSubObjects);
 
+	/**
+	 * Determines whether the specified ActorComponent is an instanced actor component and whether or not it is desirable
+	 * for the Level Snapshots process based on its creation method, this function's parameters and editor settings.
+	 * @param ActorComponent The ActorComponent to check.
+	 * @param bHideConstructionScriptComponentsInDetailsView Whether to hide construction script components in the details view. Default is true.
+	 * @return true if the ActorComponent should be added to a given Level Snapshots component array, false otherwise.
+	 */
+	bool IsInstancedActorComponentDesirableForComponentArray(
+		UActorComponent* ActorComponent, const bool bHideConstructionScriptComponentsInDetailsView = true);
+
+	/**
+	 * Filters out invalid and instanced actor components based on editor settings.
+	 * @param UnfilteredActorComponentList The unfiltered list of actor components.
+	 * @return The filtered set of actor components.
+	 */
+	TSet<UActorComponent*> FilterOutActorComponentsBasedOnEditorSettings(const TSet<UActorComponent*>& UnfilteredActorComponentList);
+
+	/**
+	 * Filters out invalid and instanced actor components based on editor settings.
+	 * @param UnfilteredActorComponentList The unfiltered list of actor components.
+	 * @return The filtered set of actor components.
+	 */
+	TSet<UActorComponent*> FilterOutActorComponentsBasedOnEditorSettings(const TSet<TWeakObjectPtr<UActorComponent>>& UnfilteredActorComponentList);
+
 	int32 CreateNewHierarchyStructInLoop(const AActor* InActor, USceneComponent* SceneComponent, TArray<TWeakPtr<FComponentHierarchy>>& AllHierarchies);
 
 	/* Creates a node tree of all scene components in an actor. Only scene components can have children. Non-scene actor components do not */
diff --git a/Engine/Plugins/VirtualProduction/LevelSnapshots/Source/LevelSnapshotsEditor/Private/Views/Results/LevelSnapshotsEditorResultsRow.cpp b/Engine/Plugins/VirtualProduction/LevelSnapshots/Source/LevelSnapshotsEditor/Private/Views/Results/LevelSnapshotsEditorResultsRow.cpp
--- a/Engine/Plugins/VirtualProduction/LevelSnapshots/Source/LevelSnapshotsEditor/Private/Views/Results/LevelSnapshotsEditorResultsRow.cpp
+++ b/Engine/Plugins/VirtualProduction/LevelSnapshots/Source/LevelSnapshotsEditor/Private/Views/Results/LevelSnapshotsEditorResultsRow.cpp
@@ -273,27 +273,16 @@ void FLevelSnapshotsEditorResultsRow::GenerateModifiedActorGroupChildren(FProper
 	const UE::LevelSnapshots::FAddedAndRemovedComponentInfo* AddedAndRemovedComponents = PropertySelectionMap.GetObjectSelection(WorldActorLocal).GetComponentSelection();
 
 	TSet<UActorComponent*> WorldComponentsToRemove;
-	if (AddedAndRemovedComponents)
-	{
-		for (TWeakObjectPtr<UActorComponent> ComponentToRemove : AddedAndRemovedComponents->EditorWorldComponentsToRemove)
-		{
-			if (ComponentToRemove.IsValid())
-			{
-				WorldComponentsToRemove.Add(ComponentToRemove.Get());
-			}
-		}
-	}
-
 	TSet<UActorComponent*> SnapshotComponentsToAdd;
 	if (AddedAndRemovedComponents)
 	{
-		for (TWeakObjectPtr<UActorComponent> ComponentToAdd : AddedAndRemovedComponents->SnapshotComponentsToAdd)
-		{
-			if (ComponentToAdd.IsValid())
-			{
-				SnapshotComponentsToAdd.Add(ComponentToAdd.Get());
-			}
-		}
+		WorldComponentsToRemove =
+			FLevelSnapshotsEditorResultsHelpers::FilterOutActorComponentsBasedOnEditorSettings(
+				AddedAndRemovedComponents->EditorWorldComponentsToRemove);
+		
+		SnapshotComponentsToAdd =
+			FLevelSnapshotsEditorResultsHelpers::FilterOutActorComponentsBasedOnEditorSettings(
+				AddedAndRemovedComponents->SnapshotComponentsToAdd);
 	}
 
 	// Iterate over components and subobjects
