Remove private include to GameplayMediaEncoder

diff --git a/Engine/Plugins/VirtualProduction/LevelSnapshots/Source/LevelSnapshotsEditor/LevelSnapshotsEditor.build.cs b/Engine/Plugins/VirtualProduction/LevelSnapshots/Source/LevelSnapshotsEditor/LevelSnapshotsEditor.build.cs
--- a/Engine/Plugins/VirtualProduction/LevelSnapshots/Source/LevelSnapshotsEditor/LevelSnapshotsEditor.build.cs
+++ b/Engine/Plugins/VirtualProduction/LevelSnapshots/Source/LevelSnapshotsEditor/LevelSnapshotsEditor.build.cs
@@ -45,11 +45,5 @@ public class LevelSnapshotsEditor : ModuleRules
 				"WorkspaceMenuStructure"
 			}
 			);
-
-		// TODO: Should not be including private headers from a different module
-		PrivateIncludePaths.AddRange(
-			new string[] {
-				Path.Combine(GetModuleDirectory("GameplayMediaEncoder"), "Private"), // For GameplayMediaEncoderCommon.h
-			});
 	}
 }
diff --git a/Engine/Plugins/VirtualProduction/LevelSnapshots/Source/LevelSnapshotsEditor/Private/LevelSnapshotsEditorFunctionLibrary.cpp b/Engine/Plugins/VirtualProduction/LevelSnapshots/Source/LevelSnapshotsEditor/Private/LevelSnapshotsEditorFunctionLibrary.cpp
--- a/Engine/Plugins/VirtualProduction/LevelSnapshots/Source/LevelSnapshotsEditor/Private/LevelSnapshotsEditorFunctionLibrary.cpp
+++ b/Engine/Plugins/VirtualProduction/LevelSnapshots/Source/LevelSnapshotsEditor/Private/LevelSnapshotsEditorFunctionLibrary.cpp
@@ -6,7 +6,6 @@
 
 #include "Editor.h"
 #include "Engine/World.h"
-#include "GameplayMediaEncoderCommon.h"
 #include "LevelSnapshot.h"
 #include "Logging/MessageLog.h"
 #include "UObject/MetaData.h"
diff --git a/Engine/Plugins/VirtualProduction/LevelSnapshots/Source/LevelSnapshotsEditor/Private/Util/TakeSnapshotUtil.cpp b/Engine/Plugins/VirtualProduction/LevelSnapshots/Source/LevelSnapshotsEditor/Private/Util/TakeSnapshotUtil.cpp
--- a/Engine/Plugins/VirtualProduction/LevelSnapshots/Source/LevelSnapshotsEditor/Private/Util/TakeSnapshotUtil.cpp
+++ b/Engine/Plugins/VirtualProduction/LevelSnapshots/Source/LevelSnapshotsEditor/Private/Util/TakeSnapshotUtil.cpp
@@ -17,7 +17,6 @@
 #include "Editor.h"
 #include "Engine/World.h"
 #include "Framework/Notifications/NotificationManager.h"
-#include "GameplayMediaEncoderCommon.h"
 #include "IAssetTools.h"
 #include "Misc/FeedbackContext.h"
 #include "Misc/MessageDialog.h"
