Fix not setting UEVer and EngineVer when loading compressed snapshot data

diff --git a/Engine/Plugins/VirtualProduction/LevelSnapshots/Source/LevelSnapshots/Private/Data/Compression/CompressionUtil.cpp b/Engine/Plugins/VirtualProduction/LevelSnapshots/Source/LevelSnapshots/Private/Data/Compression/CompressionUtil.cpp
--- a/Engine/Plugins/VirtualProduction/LevelSnapshots/Source/LevelSnapshots/Private/Data/Compression/CompressionUtil.cpp
+++ b/Engine/Plugins/VirtualProduction/LevelSnapshots/Source/LevelSnapshots/Private/Data/Compression/CompressionUtil.cpp
@@ -117,6 +117,8 @@ namespace UE::LevelSnapshots
 		const FSharedBuffer UncompressedData = CompressedData.Decompress();
 		FBufferReader UncompressedDataArchive(const_cast<void*>(UncompressedData.GetData()), UncompressedData.GetSize(), false, Ar.IsPersistent());
 		// Needed so the compressed data can migrate
+		UncompressedDataArchive.SetUEVer(Ar.UEVer());
+		UncompressedDataArchive.SetEngineVer(Ar.EngineVer());
 		UncompressedDataArchive.SetCustomVersions(Ar.GetCustomVersions());
 		UncompressedDataArchive.SetIsLoading(true);
 		Private::FExtractNamesAndObjectPathsToArraysArchive ProxyArchive(UncompressedDataArchive, MoveTemp(Names), MoveTemp(Paths));
