Propagate version info from compressed, internal LS properties to the root archive

diff --git a/Engine/Plugins/VirtualProduction/LevelSnapshots/Source/LevelSnapshots/Private/Data/Compression/CompressionUtil.cpp b/Engine/Plugins/VirtualProduction/LevelSnapshots/Source/LevelSnapshots/Private/Data/Compression/CompressionUtil.cpp
--- a/Engine/Plugins/VirtualProduction/LevelSnapshots/Source/LevelSnapshots/Private/Data/Compression/CompressionUtil.cpp
+++ b/Engine/Plugins/VirtualProduction/LevelSnapshots/Source/LevelSnapshots/Private/Data/Compression/CompressionUtil.cpp
@@ -9,6 +9,7 @@
 #include "Compression/CompressedBuffer.h"
 #include "Serialization/ArchiveSerializedPropertyChain.h"
 #include "Serialization/BufferArchive.h"
+#include "Serialization/CustomVersion.h"
 
 namespace UE::LevelSnapshots
 {
@@ -68,9 +69,10 @@ namespace UE::LevelSnapshots
 			}
 		};
 
-		FCompressedBuffer CompressData(const FSharedBuffer& RawData, const ECompressedBufferCompressor Compressor, const ECompressedBufferCompressionLevel CompressionLevel);
-		void SerializeUncompressableProperties(FArchive& Archive, FWorldSnapshotData& Data);
-		void SerializeCompressableProperties(FArchive& Archive, FWorldSnapshotData& Data);
+		static FCompressedBuffer CompressData(const FSharedBuffer& RawData, const ECompressedBufferCompressor Compressor, const ECompressedBufferCompressionLevel CompressionLevel);
+		static void SerializeUncompressableProperties(FArchive& Archive, FWorldSnapshotData& Data);
+		static void SerializeCompressableProperties(FArchive& Archive, FWorldSnapshotData& Data);
+		static void CopyVersionInfo(const FArchive& Source, FArchive& Target);
 	}
 	
 	void Compress(FArchive& Ar, FWorldSnapshotData& Data)
@@ -85,6 +87,8 @@ namespace UE::LevelSnapshots
 		UncompressedDataArchive.SetIsSaving(true);
 		Private::FExtractNamesAndObjectPathsToArraysArchive ProxyArchive(UncompressedDataArchive);
 		Private::SerializeCompressableProperties(ProxyArchive, Data);
+		// The compressed properties may have added version info themselves which must be propagated to the root archive.
+		Private::CopyVersionInfo(ProxyArchive, Ar);
 
 		// 2. Auxiliary data - harvesting stage of save system requires this
 		Ar << ProxyArchive.Names;
@@ -116,10 +120,13 @@ namespace UE::LevelSnapshots
 		const FCompressedBuffer CompressedData = FCompressedBuffer::Load(Ar);
 		const FSharedBuffer UncompressedData = CompressedData.Decompress();
 		FBufferReader UncompressedDataArchive(const_cast<void*>(UncompressedData.GetData()), UncompressedData.GetSize(), false, Ar.IsPersistent());
+		
 		// Needed so the compressed data can migrate
 		UncompressedDataArchive.SetUEVer(Ar.UEVer());
 		UncompressedDataArchive.SetEngineVer(Ar.EngineVer());
-		UncompressedDataArchive.SetCustomVersions(Ar.GetCustomVersions());
+		// This will also include versions propagated by CopyVersionInfo during save
+		UncompressedDataArchive.SetCustomVersions(Ar.GetCustomVersions()); 
+		
 		UncompressedDataArchive.SetIsLoading(true);
 		Private::FExtractNamesAndObjectPathsToArraysArchive ProxyArchive(UncompressedDataArchive, MoveTemp(Names), MoveTemp(Paths));
 		Private::SerializeCompressableProperties(ProxyArchive, Data);
@@ -127,7 +134,7 @@ namespace UE::LevelSnapshots
 
 	namespace Private
 	{
-		FCompressedBuffer CompressData(const FSharedBuffer& RawData, const ECompressedBufferCompressor Compressor, const ECompressedBufferCompressionLevel CompressionLevel)
+		static FCompressedBuffer CompressData(const FSharedBuffer& RawData, const ECompressedBufferCompressor Compressor, const ECompressedBufferCompressionLevel CompressionLevel)
 		{
 			SCOPED_SNAPSHOT_CORE_TRACE(Compress);
 			return FCompressedBuffer::Compress(RawData, Compressor, CompressionLevel);
@@ -177,7 +184,7 @@ namespace UE::LevelSnapshots
 			};
 		}
 
-		void SerializeUncompressableProperties(FArchive& Archive, FWorldSnapshotData& Data)
+		static void SerializeUncompressableProperties(FArchive& Archive, FWorldSnapshotData& Data)
 		{
 			// Names and References are collected during the harvesting state when a packet is saved
 			// References must be saved normally. It seems like Names could be safe to compress, too, but the gain was minimal when tested a 4k actor map... let's not take risks for little gain...
@@ -197,11 +204,19 @@ namespace UE::LevelSnapshots
 			FWorldSnapshotData::StaticStruct()->SerializeTaggedProperties(FStructuredArchiveFromArchive(ArchiveProxy).GetSlot(), (uint8*)&Data, FWorldSnapshotData::StaticStruct(), nullptr);
 		}
 		
-		void SerializeCompressableProperties(FArchive& Archive, FWorldSnapshotData& Data)
+		static void SerializeCompressableProperties(FArchive& Archive, FWorldSnapshotData& Data)
 		{
 			FWorldDataArchive ArchiveProxy(Archive, GetCompressedProperties());
 			// SerializeTaggedProperties for forward compatibility, e.g. renaming properties
 			FWorldSnapshotData::StaticStruct()->SerializeTaggedProperties(FStructuredArchiveFromArchive(ArchiveProxy).GetSlot(), (uint8*)&Data, FWorldSnapshotData::StaticStruct(), nullptr);
 		}
+
+		static void CopyVersionInfo(const FArchive& Source, FArchive& Target)
+		{
+			for (const FCustomVersion& Version : Source.GetCustomVersions().GetAllVersions())
+			{
+				Target.UsingCustomVersion(Version.Key);
+			}
+		}
 	}
 }
