Submitting attempted fix of hard to reproduce crash when expanding Level Snapshots row on behalf of Jared.Therriault.
- The attempted fix is adding a IPropertyHandle::IsValidHandle check
- This also replaces initing a const& to a temporary return value of a function.  While this does not invoke undefined behaviour per se,  it dangerous to leave in because a developer can by accident extend the lifetime of the reference and move into UB territory.

diff --git a/Engine/Plugins/VirtualProduction/LevelSnapshots/Source/LevelSnapshotsEditor/Private/Views/Results/SLevelSnapshotsEditorResultsRow.cpp b/Engine/Plugins/VirtualProduction/LevelSnapshots/Source/LevelSnapshotsEditor/Private/Views/Results/SLevelSnapshotsEditorResultsRow.cpp
--- a/Engine/Plugins/VirtualProduction/LevelSnapshots/Source/LevelSnapshotsEditor/Private/Views/Results/SLevelSnapshotsEditorResultsRow.cpp
+++ b/Engine/Plugins/VirtualProduction/LevelSnapshots/Source/LevelSnapshotsEditor/Private/Views/Results/SLevelSnapshotsEditorResultsRow.cpp
@@ -267,40 +267,43 @@ void SLevelSnapshotsEditorResultsRow::GenerateAddedAndRemovedRowComponents(
 }
 
 TSharedRef<SWidget> SLevelSnapshotsEditorResultsRow::GenerateFinalValueWidget(
-	const ELevelSnapshotsObjectType InObjectType, FLevelSnapshotsEditorResultsRowPtr PinnedItem, const bool bIsHeaderRow, const bool bNeedsNullWidget) const
+	const ELevelSnapshotsObjectType InObjectType,
+	FLevelSnapshotsEditorResultsRowPtr PinnedItem,
+	const bool bIsHeaderRow,
+	const bool bNeedsNullWidget
+	) const
 {
 	// Nested Splitter Slot 0
 	const int32 SlotIndex = InObjectType == ELevelSnapshotsObjectType::ObjectType_Snapshot ? 1 : 0;
 	
 	TSharedPtr<SWidget> ChildWidget;
-
 	bool bIsChildWidgetCustomized = false;
 
-	if (const TSharedPtr<IPropertyHandle>& PropertyHandle = 
-		InObjectType == ELevelSnapshotsObjectType::ObjectType_Snapshot ? PinnedItem->GetSnapshotPropertyHandle() : PinnedItem->GetWorldPropertyHandle())
+	const TSharedPtr<IPropertyHandle> PropertyHandle = InObjectType == ELevelSnapshotsObjectType::ObjectType_Snapshot
+		? PinnedItem->GetSnapshotPropertyHandle()
+		: PinnedItem->GetWorldPropertyHandle();
+	if (PropertyHandle && PropertyHandle->IsValidHandle())
 	{
 		if (const TSharedPtr<IPropertyHandle> KeyHandle = PropertyHandle->GetKeyHandle())
 		{
-			TSharedRef<SSplitter> Splitter = SNew(SSplitter).ResizeMode(ESplitterResizeMode::FixedPosition);
-
+			const TSharedRef<SSplitter> Splitter = SNew(SSplitter).ResizeMode(ESplitterResizeMode::FixedPosition);
 			if (KeyHandle->IsValidHandle())
 			{
 				Splitter->AddSlot()[KeyHandle->CreatePropertyValueWidget(false)];
 			}
 
 			Splitter->AddSlot()[PropertyHandle->CreatePropertyValueWidget(false)];
-
 			ChildWidget = Splitter;
 		}
 		else
-		{				
-			if (const TSharedPtr<IDetailTreeNode>& Node = 
-				InObjectType == ELevelSnapshotsObjectType::ObjectType_Snapshot ? PinnedItem->GetSnapshotPropertyNode() : PinnedItem->GetWorldPropertyNode())
+		{
+			const TSharedPtr<IDetailTreeNode> Node = InObjectType == ELevelSnapshotsObjectType::ObjectType_Snapshot
+				? PinnedItem->GetSnapshotPropertyNode()
+				: PinnedItem->GetWorldPropertyNode();
+			if (Node)
 			{
 				bIsChildWidgetCustomized = true;
-				
 				const FNodeWidgets Widgets = Node->CreateNodeWidgets();
-
 				ChildWidget = Widgets.WholeRowWidget.IsValid() ? Widgets.WholeRowWidget : Widgets.ValueWidget;
 			}
 			else
