Fix crash / failing check when deleting snapshot containing sub-level data.

When the level snapshot is about to be deleted, the internal, transient snapshot world is now cleaned up.
This ensures objects contained in the world, like subsystems, are cleaned up in the order they expect; otherwise GC would delete the objects in a non-deterministic order, causing failing checks.

diff --git a/Engine/Plugins/VirtualProduction/LevelSnapshots/Source/LevelSnapshots/Private/Data/LevelSnapshot.cpp b/Engine/Plugins/VirtualProduction/LevelSnapshots/Source/LevelSnapshots/Private/Data/LevelSnapshot.cpp
--- a/Engine/Plugins/VirtualProduction/LevelSnapshots/Source/LevelSnapshots/Private/Data/LevelSnapshot.cpp
+++ b/Engine/Plugins/VirtualProduction/LevelSnapshots/Source/LevelSnapshots/Private/Data/LevelSnapshot.cpp
@@ -23,6 +23,9 @@
 #include "GameFramework/Actor.h"
 #include "Misc/ScopeExit.h"
 #include "UObject/Package.h"
+#if WITH_EDITOR
+#include "Editor.h"
+#endif
 #if WITH_EDITOR && !(UE_BUILD_SHIPPING || UE_BUILD_TEST)
 #include "Logging/MessageLog.h"
 #endif
@@ -393,6 +396,33 @@ void ULevelSnapshot::EnsureWorldInitialised()
 		Cache.InitFor(SerializedData);
 #if WITH_EDITOR
 		FCoreUObjectDelegates::OnObjectModified.AddUObject(this, &ULevelSnapshot::ResetDiffCacheToUninitialized);
+
+		// Deleting assets will cause GC to destroy objects in the order they appear in GUnreachableObjects.
+		// It's possible that objects in one of our UWorlds appear before the actual UWorlds. So ConditionalBeginDestroy is invoked on objects before the UWorld instance.
+		// Problem: Certain objects expect UWorld::CleanupWorld to be called before their own BeginDestroy (e.g. subsystems). Not doing so causes all kinds of failing checks and ensures.
+		// Solution: We pre-emptively clean up the UWorld before actual deletion; this triggers important destruction & clean-up callbacks (e.g for subsystems, rendering, etc.).
+		// Caveat of solution: This code should probably run with FCoreUObjectDelegates::PreGarbageCollectConditionalBeginDestroy instead. However,
+		//	1. when that callback executes, there is no API for knowing whether our LevelSnapshot is in GUnreachableObjects.
+		//	2. user code in actors may call CollectGarbage while applying a snapshot. Calling DestroyWorld in that context can be disastrous.
+		// Decision: Now (= 5.4.2), the below events are the only known cases where our snapshot is deleted... hence we do not use PreGarbageCollectConditionalBeginDestroy due to lack of querying API.
+		// For future reference: If the PreGarbageCollectConditionalBeginDestroy API ever returns the list of objects to be destroyed, then you should
+		// replace OnAssetsPreDelete and OnPreDestructiveAssetAction with PreGarbageCollectConditionalBeginDestroy.
+		FEditorDelegates::OnAssetsPreDelete.AddWeakLambda(this, [this](const TArray<UObject*>& ObjectsAboutToDelete)
+		{
+			const bool bHasInitializedWorld = RootSnapshotWorld != nullptr;
+			if (bHasInitializedWorld && ObjectsAboutToDelete.Contains(this))
+			{
+				DestroyWorld();
+			}
+		});
+		FEditorDelegates::OnPreDestructiveAssetAction.AddWeakLambda(this, [this](const TArray<UObject*>& ObjectsAboutToDelete, EDestructiveAssetActions, FResultMessage&)
+		{
+			const bool bHasInitializedWorld = RootSnapshotWorld != nullptr;
+			if (bHasInitializedWorld && ObjectsAboutToDelete.Contains(this))
+			{
+				DestroyWorld();
+			}
+		});
 #endif
 	}
 
@@ -450,11 +480,12 @@ void ULevelSnapshot::DestroyWorld()
 			GEngine->OnWorldDestroyed().Remove(Handle);
 			Handle.Reset();
 		}
+#if WITH_EDITOR
+		FEditorDelegates::OnAssetsPreDelete.RemoveAll(this);
+		FEditorDelegates::OnPreDestructiveAssetAction.RemoveAll(this);
+#endif
 		
 		CleanUpWorld();
-		TransientWorldPackage = nullptr;
-		RootSnapshotWorld = nullptr;
-		SnapshotSublevels.Reset();
 	}
 }
 
@@ -466,12 +497,17 @@ void ULevelSnapshot::CleanUpWorld()
 
 	// Cleanup world clears RF_Standalone flag...
 	const bool bWasStandalone = HasAnyFlags(RF_Standalone);
+	
 	RootSnapshotWorld->CleanupWorld();
 	for (UWorld* Sublevel : SnapshotSublevels)
 	{
 		Sublevel->CleanupWorld();
 	}
-	
+	TransientWorldPackage = nullptr;
+	RootSnapshotWorld = nullptr;
+	SnapshotSublevels.Reset();
+
+	// ... so restore it
 	if (bWasStandalone)
 	{
 		SetFlags(RF_Standalone);
