Fix missing save icons in save preset combo menu of Level Snapshots

diff --git a/Engine/Plugins/VirtualProduction/LevelSnapshots/Source/LevelSnapshotsEditor/Private/Widgets/Filter/SSaveAndLoadFilters.cpp b/Engine/Plugins/VirtualProduction/LevelSnapshots/Source/LevelSnapshotsEditor/Private/Widgets/Filter/SSaveAndLoadFilters.cpp
--- a/Engine/Plugins/VirtualProduction/LevelSnapshots/Source/LevelSnapshotsEditor/Private/Widgets/Filter/SSaveAndLoadFilters.cpp
+++ b/Engine/Plugins/VirtualProduction/LevelSnapshots/Source/LevelSnapshotsEditor/Private/Widgets/Filter/SSaveAndLoadFilters.cpp
@@ -127,7 +127,7 @@ TSharedRef<SWidget> SSaveAndLoadFilters::GenerateSaveLoadMenu()
 		MenuBuilder.AddMenuEntry(
             EntryName,
             LOCTEXT("SaveExistingFiltersToolTip", "Overwrite the asset you last loaded"),
-            FSlateIcon(FAppStyle::Get().GetStyleSetName(), "AssetEditor.SaveAsset.Greyscale"),
+            FSlateIcon(FAppStyle::Get().GetStyleSetName(), "AssetEditor.SaveAsset"),
             FUIAction(
                 FExecuteAction::CreateStatic(&SaveExisting, FilterLoader)
             )
@@ -137,7 +137,7 @@ TSharedRef<SWidget> SSaveAndLoadFilters::GenerateSaveLoadMenu()
 	MenuBuilder.AddMenuEntry(
 	    LOCTEXT("SaveFiltersAs", "Save as..."),
 	    LOCTEXT("SaveFiltersAsToolTip", "Saves a new asset."),
-	    FSlateIcon(FAppStyle::Get().GetStyleSetName(), "AssetEditor.SaveAsset.Greyscale"),
+	    FSlateIcon(FAppStyle::Get().GetStyleSetName(), "AssetEditor.SaveAssetAs"),
 	    FUIAction(
 	        FExecuteAction::CreateStatic(&SaveAs, FilterLoader)
 		)	
