Fix Level Snapshots serializing references to CDO subobjects as if they subobjects in a UWorld.

diff --git a/Engine/Plugins/VirtualProduction/LevelSnapshots/Source/LevelSnapshots/Private/Data/Util/ObjectDependencyUtil.cpp b/Engine/Plugins/VirtualProduction/LevelSnapshots/Source/LevelSnapshots/Private/Data/Util/ObjectDependencyUtil.cpp
--- a/Engine/Plugins/VirtualProduction/LevelSnapshots/Source/LevelSnapshots/Private/Data/Util/ObjectDependencyUtil.cpp
+++ b/Engine/Plugins/VirtualProduction/LevelSnapshots/Source/LevelSnapshots/Private/Data/Util/ObjectDependencyUtil.cpp
@@ -364,7 +364,11 @@ int32 UE::LevelSnapshots::Private::AddObjectDependency(FWorldSnapshotData& World
 	// Even if FSnapshotRestorability::IsSubobjectDesirableForCapture later returns false for this object, we want to track it
 	const int32 Result = Internal::AddOrFindObjectReference(WorldData, ReferenceFromOriginalObject);
 
-	if (bCheckWhetherSubobject && ReferenceFromOriginalObject && ReferenceFromOriginalObject->GetTypedOuter<AActor>())
+	const AActor* OwningActor = bCheckWhetherSubobject && ReferenceFromOriginalObject ? ReferenceFromOriginalObject->GetTypedOuter<AActor>() : nullptr;
+	if (OwningActor
+		// This handles the rare case that ReferenceFromOriginalObject is a reference to a subobject of a CDO actor or in a Blueprint.
+		// In that case, ReferenceFromOriginalObject is not actually a real world object, hence not a subobject we want to save.
+		&& !OwningActor->HasAnyFlags(RF_ClassDefaultObject | RF_ArchetypeObject))
 	{
 		Internal::AddSubobjectDependencyInternal(WorldData, ReferenceFromOriginalObject, Result);
 	}
