Fix crash in when diffing a snapshot against an object with an TSet or TMap UPROPERTY(Instanced).
This crash occured if a user had saved a snapshot with objects in the TSet/TMap and then removed an entry.

diff --git a/Engine/Plugins/VirtualProduction/LevelSnapshots/Source/LevelSnapshots/Private/Filtering/Diffing/ApplySnapshotFilter.cpp b/Engine/Plugins/VirtualProduction/LevelSnapshots/Source/LevelSnapshots/Private/Filtering/Diffing/ApplySnapshotFilter.cpp
--- a/Engine/Plugins/VirtualProduction/LevelSnapshots/Source/LevelSnapshots/Private/Filtering/Diffing/ApplySnapshotFilter.cpp
+++ b/Engine/Plugins/VirtualProduction/LevelSnapshots/Source/LevelSnapshots/Private/Filtering/Diffing/ApplySnapshotFilter.cpp
@@ -498,12 +498,17 @@ TOptional<UE::LevelSnapshots::Private::FApplySnapshotFilter::EPropertySearchResu
 			}
 		}
 
-		void* GetSnapshotValuePtrFromIndex(int32 WorldIndex)
+		void* GetSnapshotValuePtrFromIndex(int32 LogicalWorldIndex)
 		{
 			// Note: we only look for subobjects with the same name because it is faster; otherwise we'd need to compare all properties (n^2 performance).
 			// This works most of the time but not when some subobject has changed name (unlikely but possible).
 			// In the cases it does not work, we just end up with an additional diff in the results view; a low price to pay.
-			const void* WorldValuePtr = GetWorldValuePtrFromIndex(WorldIndex);
+			const void* WorldValuePtr = GetWorldValuePtrFromIndex(LogicalWorldIndex);
+			if (!WorldValuePtr)
+			{
+				return nullptr;
+			}
+			
 			if (UObject* WorldObject = ObjectProperty->GetObjectPropertyValue(WorldValuePtr))
 			{
 				void** PossibleSnapshotObject = ObjectNameToSnapshotValuePtr.Find(WorldObject->GetFName());
@@ -512,9 +517,10 @@ TOptional<UE::LevelSnapshots::Private::FApplySnapshotFilter::EPropertySearchResu
 			return nullptr;
 		}
 
-		void* GetWorldValuePtrFromIndex(int32 WorldIndex)
+		void* GetWorldValuePtrFromIndex(int32 LogicalWorldIndex)
 		{
-			return WorldSet.GetElementPtr(WorldIndex);
+			const int32 InternalIndex = WorldSet.FindInternalIndex(LogicalWorldIndex);
+			return WorldSet.IsValidIndex(InternalIndex) ? WorldSet.GetElementPtr(InternalIndex) : nullptr;
 		}
 	};
 
@@ -544,15 +550,21 @@ TOptional<UE::LevelSnapshots::Private::FApplySnapshotFilter::EPropertySearchResu
 			WorldNum(WorldMap.Num())
 		{}
 
-		void* GetSnapshotValuePtrFromIndex(int32 WorldIndex)
+		void* GetSnapshotValuePtrFromIndex(int32 LogicalWorldIndex)
 		{
-			const void* WorldKey = WorldMap.GetKeyPtr(WorldIndex);
-			return SnapshotMap.FindValueFromHash(WorldKey);
+			const int32 InternalIndex = WorldMap.FindInternalIndex(LogicalWorldIndex);
+			if (SnapshotMap.IsValidIndex(InternalIndex))
+			{
+				const void* WorldKey = WorldMap.GetKeyPtr(InternalIndex);
+				return SnapshotMap.FindValueFromHash(WorldKey);
+			}
+			return nullptr;
 		}
 
-		void* GetWorldValuePtrFromIndex(int32 WorldIndex)
+		void* GetWorldValuePtrFromIndex(int32 LogicalWorldIndex)
 		{
-			return WorldMap.GetValuePtr(WorldIndex);
+			const int32 InternalIndex = WorldMap.FindInternalIndex(LogicalWorldIndex);
+			return WorldMap.IsValidIndex(InternalIndex) ? WorldMap.GetValuePtr(InternalIndex) : nullptr;
 		}
 	};
 
@@ -575,20 +587,20 @@ TOptional<UE::LevelSnapshots::Private::FApplySnapshotFilter::EPropertySearchResu
 	for (int32 WorldIndex = 0; WorldIndex < MaxSharedSize; ++WorldIndex)
 	{
 		void* SnapshotValuePtr = Detail.GetSnapshotValuePtrFromIndex(WorldIndex);
-		if (!SnapshotValuePtr)
+		void* WorldValuePtr = SnapshotValuePtr ? Detail.GetWorldValuePtrFromIndex(WorldIndex) : nullptr;
+		if (!SnapshotValuePtr || !WorldValuePtr)
 		{
 			bHadUnmatchedObjects = true;
 			continue;
 		}
-
-		void* WorldValuePtr = Detail.GetWorldValuePtrFromIndex(WorldIndex);
-		const TOptional<EPropertySearchResult> SubbjectPairResult = TrackPossibleSubobjectProperties(ContainerContext, ObjectProperty, SnapshotValuePtr, WorldValuePtr);
-		if (!SubbjectPairResult.IsSet())
+		
+		const TOptional<EPropertySearchResult> SubobjectPairResult = TrackPossibleSubobjectProperties(ContainerContext, ObjectProperty, SnapshotValuePtr, WorldValuePtr);
+		if (!SubobjectPairResult.IsSet())
 		{
 			return {};
 		}
 		
-		const bool bFoundChangesOnSubobject = SubbjectPairResult == EPropertySearchResult::FoundProperties;
+		const bool bFoundChangesOnSubobject = SubobjectPairResult == EPropertySearchResult::FoundProperties;
 		bFoundChangedProperties |= bFoundChangesOnSubobject;
 		UObject* WorldObject = ObjectProperty->GetObjectPropertyValue(WorldValuePtr);
 		if (!bFoundChangesOnSubobject && WorldObject)
diff --git a/Engine/Plugins/VirtualProduction/LevelSnapshots/Source/LevelSnapshots/Private/Tests/RegressionTests.cpp b/Engine/Plugins/VirtualProduction/LevelSnapshots/Source/LevelSnapshots/Private/Tests/RegressionTests.cpp
--- a/Engine/Plugins/VirtualProduction/LevelSnapshots/Source/LevelSnapshots/Private/Tests/RegressionTests.cpp
+++ b/Engine/Plugins/VirtualProduction/LevelSnapshots/Source/LevelSnapshots/Private/Tests/RegressionTests.cpp
@@ -279,4 +279,67 @@ namespace UE::LevelSnapshots::Private::Tests
 		return true;
 	}
 #endif
+
+	/**
+	 * There used to be a bug in FApplySnapshotFilter::TrackPossibleMapSubobjectProperties would invoke FScriptSetHelper:.GetElementPtr and
+	 * FScriptMapHelper::GetKeyPtr before converting logical indices to internal indices.
+	 * The result was that the following would crash:
+	 * 1. Add an actor that has a UPROPERTY(EditAnywhere, Instanced) TMap<FName, UObject*> property
+	 * 2. Add 3 subobjects to that property
+	 * 3. Remove the middle one
+	 * 4. Diff the snapshot > Crash
+	 */
+	IMPLEMENT_SIMPLE_AUTOMATION_TEST(FInstancedSubobjectRemovedFromTMap, "VirtualProduction.LevelSnapshots.Snapshot.Regression.InstancedSubobjectRemovedFromTMap", (EAutomationTestFlags_ApplicationContextMask | EAutomationTestFlags::EngineFilter));
+	bool FInstancedSubobjectRemovedFromTMap::RunTest(const FString& Parameters)
+	{
+		ASnapshotTestActor* Actor = nullptr;
+		USubobject* Subobject1 = nullptr;
+		USubobject* Subobject2 = nullptr;
+		USubobject* Subobject3 = nullptr;
+		const FName SubobjectName1 = TEXT("Subobject1");
+		const FName SubobjectName2 = TEXT("Subobject2");
+		const FName SubobjectName3 = TEXT("Subobject3");
+		
+		FSnapshotTestRunner()
+			.ModifyWorld([&](UWorld* World)
+			{
+				Actor = ASnapshotTestActor::Spawn(World);
+				Subobject1 = NewObject<USubobject>(Actor, SubobjectName1);
+				Subobject2 = NewObject<USubobject>(Actor, SubobjectName2);
+				Subobject3 = NewObject<USubobject>(Actor, SubobjectName3);
+
+				Subobject1->IntProperty = 1;
+				Subobject2->IntProperty = 2;
+				Subobject3->IntProperty = 3;
+				
+				Actor->EditableInstancedSubobjectMap_OptionalSubobject.Add(SubobjectName1, Subobject1);
+				Actor->EditableInstancedSubobjectMap_OptionalSubobject.Add(SubobjectName2, Subobject2);
+				Actor->EditableInstancedSubobjectMap_OptionalSubobject.Add(SubobjectName3, Subobject3);
+			})
+			.TakeSnapshot()
+			.ModifyWorld([&](UWorld* World)
+			{
+				Actor->EditableInstancedSubobjectMap_OptionalSubobject.Remove(SubobjectName2);
+			})
+			.ApplySnapshot()
+			.ModifyWorld([&](UWorld* World)
+			{
+				const TObjectPtr<USubobject>* Restored1 = Actor->EditableInstancedSubobjectMap_OptionalSubobject.Find(SubobjectName1);
+				const TObjectPtr<USubobject>* Restored2 = Actor->EditableInstancedSubobjectMap_OptionalSubobject.Find(SubobjectName2);
+				const TObjectPtr<USubobject>* Restored3 = Actor->EditableInstancedSubobjectMap_OptionalSubobject.Find(SubobjectName3);
+
+				if (!Restored1 || !Restored2 || !Restored3 || !*Restored1 || !*Restored2 || !*Restored3)
+				{
+					AddError(TEXT("Not all objects where restored"));
+				}
+				else
+				{
+					TestEqual(TEXT("1"), Restored1->Get()->IntProperty, 1);
+					TestEqual(TEXT("2"), Restored2->Get()->IntProperty, 2);
+					TestEqual(TEXT("3"), Restored3->Get()->IntProperty, 3);
+				}
+			});
+
+		return true;
+	}
 }
\ No newline at end of file
