[LevelSnapshot] Tweaks for the overall size to better match AssetView length

diff --git a/Engine/Plugins/VirtualProduction/LevelSnapshots/Source/LevelSnapshotsEditor/Private/LevelSnapshotsEditorModule.cpp b/Engine/Plugins/VirtualProduction/LevelSnapshots/Source/LevelSnapshotsEditor/Private/LevelSnapshotsEditorModule.cpp
--- a/Engine/Plugins/VirtualProduction/LevelSnapshots/Source/LevelSnapshotsEditor/Private/LevelSnapshotsEditorModule.cpp
+++ b/Engine/Plugins/VirtualProduction/LevelSnapshots/Source/LevelSnapshotsEditor/Private/LevelSnapshotsEditorModule.cpp
@@ -96,11 +96,14 @@ void FLevelSnapshotsEditorModule::RegisterTabSpawner()
 			.SetTooltipText(NSLOCTEXT("LevelSnapshots", "LevelSnapshotsTooltipText", "Open the Level Snapshots tab"))
 			.SetIcon(FSlateIcon(FLevelSnapshotsEditorStyle::GetStyleSetName(), "LevelSnapshots.ToolbarButton", "LevelSnapshots.ToolbarButton.Small")
 			);
+	FGlobalTabmanager::Get()->RegisterDefaultTabWindowSize(LevelSnapshotsEditor::LevelSnapshotsTabName, FVector2D(1240, 600));
+	
 	TabSpawnerEntry.SetGroup(WorkspaceMenu::GetMenuStructure().GetLevelEditorVirtualProductionCategory());
 }
 
 void FLevelSnapshotsEditorModule::UnregisterTabSpawner()
 {
+	FGlobalTabmanager::Get()->UnregisterDefaultTabWindowSize(LevelSnapshotsEditor::LevelSnapshotsTabName);
 	FGlobalTabmanager::Get()->UnregisterNomadTabSpawner(LevelSnapshotsEditor::LevelSnapshotsTabName);
 }
 
diff --git a/Engine/Plugins/VirtualProduction/LevelSnapshots/Source/LevelSnapshotsEditor/Private/Views/SLevelSnapshotsEditor.cpp b/Engine/Plugins/VirtualProduction/LevelSnapshots/Source/LevelSnapshotsEditor/Private/Views/SLevelSnapshotsEditor.cpp
--- a/Engine/Plugins/VirtualProduction/LevelSnapshots/Source/LevelSnapshotsEditor/Private/Views/SLevelSnapshotsEditor.cpp
+++ b/Engine/Plugins/VirtualProduction/LevelSnapshots/Source/LevelSnapshotsEditor/Private/Views/SLevelSnapshotsEditor.cpp
@@ -43,7 +43,7 @@ void SLevelSnapshotsEditor::Construct(const FArguments& InArgs, ULevelSnapshotsE
 	RegisterTabSpawners(TabManager.ToSharedRef(), AppMenuGroup);
 
 	// Create our content
-	const TSharedRef<FTabManager::FLayout> Layout =
+	TSharedRef<FTabManager::FLayout> Layout =
 		FTabManager::NewLayout("Levelsnapshots_Layout_2")
 		->AddArea
 		(
@@ -61,14 +61,14 @@ void SLevelSnapshotsEditor::Construct(const FArguments& InArgs, ULevelSnapshotsE
 				->Split
 				(
 					FTabManager::NewStack()
-					->SetSizeCoefficient(0.2f)
+					->SetSizeCoefficient(0.25f)
 					->AddTab(InputTabID, ETabState::OpenedTab)
 					->SetHideTabWell(true)
 					)
 				->Split
 				(
 					FTabManager::NewStack()
-					->SetSizeCoefficient(0.45f)
+					->SetSizeCoefficient(0.25f)
 					->AddTab(FilterTabID, ETabState::OpenedTab)
 					->SetHideTabWell(true)
 					)
@@ -82,10 +82,25 @@ void SLevelSnapshotsEditor::Construct(const FArguments& InArgs, ULevelSnapshotsE
 				)
 			);
 
+	Layout = FLayoutSaveRestore::LoadFromConfig(GEditorLayoutIni, Layout);
+
 	ChildSlot
 	[
 		TabManager->RestoreFrom(Layout, ConstructUnderWindow).ToSharedRef()
 	];
+
+	// on persist layout will handle saving layout if the editor is shut down:
+	TabManager->SetOnPersistLayout(
+		FTabManager::FOnPersistLayout::CreateLambda(
+			[](const TSharedRef<FTabManager::FLayout>& InLayout)
+			{
+				if (InLayout->GetPrimaryArea().Pin().IsValid())
+				{
+					FLayoutSaveRestore::SaveToConfig(GEditorLayoutIni, InLayout);
+				}
+			}
+		)
+	);
 }
 
 void SLevelSnapshotsEditor::OpenLevelSnapshotsDialogWithAssetSelected(const FAssetData& InAssetData) const
@@ -115,6 +130,7 @@ void SLevelSnapshotsEditor::RegisterTabSpawners(const TSharedRef<FTabManager>& I
 TSharedRef<SDockTab> SLevelSnapshotsEditor::SpawnTab_CustomToolbar(const FSpawnTabArgs& Args)
 {
 	return SNew(SDockTab)
+		.CanEverClose(false)
         .Label(LOCTEXT("Levelsnapshots.Toolkit.ToolbarTitle", "Toolbar"))
 		.ShouldAutosize(true)
         [
@@ -193,6 +209,7 @@ TSharedRef<SDockTab> SLevelSnapshotsEditor::SpawnTab_Input(const FSpawnTabArgs&
 	check(EditorData.IsValid());
 	
 	InputTab = SNew(SDockTab)
+		.CanEverClose(false)
 		.Label(LOCTEXT("Levelsnapshots.Toolkit.InputTitle", "Input"))
 		[
 			SAssignNew(EditorInputWidget, SLevelSnapshotsEditorInput, EditorData.Get())
@@ -206,6 +223,7 @@ TSharedRef<SDockTab> SLevelSnapshotsEditor::SpawnTab_Filter(const FSpawnTabArgs&
 	check(EditorData.IsValid());
 	
 	TSharedPtr<SDockTab> DetailsTab = SNew(SDockTab)
+		.CanEverClose(false)
 		.Label(LOCTEXT("Levelsnapshots.Toolkit.FilterTitle", "Filter"))
 		[
 			SNew(SLevelSnapshotsEditorFilters, EditorData.Get())
@@ -218,13 +236,14 @@ TSharedRef<SDockTab> SLevelSnapshotsEditor::SpawnTab_Results(const FSpawnTabArgs
 {
 	check(EditorData.IsValid());
 	
-	TSharedPtr<SDockTab> DetailsTab = SNew(SDockTab)
+	TSharedPtr<SDockTab> ResultsTab = SNew(SDockTab)
+		.CanEverClose(false)
 		.Label(LOCTEXT("Levelsnapshots.Toolkit.ResultTitle", "Result"))
 		[
 			SNew(SLevelSnapshotsEditorResults, EditorData.Get())
 		];
 
-	return DetailsTab.ToSharedRef();
+	return ResultsTab.ToSharedRef();
 }
 
 FReply SLevelSnapshotsEditor::OnClickTakeSnapshot()
diff --git a/Engine/Plugins/VirtualProduction/LevelSnapshots/Source/LevelSnapshotsEditor/Private/Widgets/SLevelSnapshotsEditorBrowser.cpp b/Engine/Plugins/VirtualProduction/LevelSnapshots/Source/LevelSnapshotsEditor/Private/Widgets/SLevelSnapshotsEditorBrowser.cpp
--- a/Engine/Plugins/VirtualProduction/LevelSnapshots/Source/LevelSnapshotsEditor/Private/Widgets/SLevelSnapshotsEditorBrowser.cpp
+++ b/Engine/Plugins/VirtualProduction/LevelSnapshots/Source/LevelSnapshotsEditor/Private/Widgets/SLevelSnapshotsEditorBrowser.cpp
@@ -54,7 +54,7 @@ void SLevelSnapshotsEditorBrowser::Construct(const FArguments& InArgs, ULevelSna
 	AssetPickerConfig.bShowPathInColumnView = true;
 	AssetPickerConfig.bShowTypeInColumnView = false;
 	AssetPickerConfig.bSortByPathInColumnView = false;
-	AssetPickerConfig.SaveSettingsName = TEXT("GlobalAssetPicker");
+	AssetPickerConfig.SaveSettingsName = TEXT("LevelSnapshotPicker");
 	AssetPickerConfig.ThumbnailScale = 0.8f;
 	AssetPickerConfig.Filter = ARFilter;
 	AssetPickerConfig.SelectionMode = ESelectionMode::Single;
