LevelSnapshots: Fixed BP function not being static.

The function was practically unusable otherwise.

#rnx

diff --git a/Engine/Plugins/VirtualProduction/LevelSnapshots/Source/LevelShapshotFilters/Public/PropertyBlueprintFunctionLibrary.h b/Engine/Plugins/VirtualProduction/LevelSnapshots/Source/LevelShapshotFilters/Public/PropertyBlueprintFunctionLibrary.h
--- a/Engine/Plugins/VirtualProduction/LevelSnapshots/Source/LevelShapshotFilters/Public/PropertyBlueprintFunctionLibrary.h
+++ b/Engine/Plugins/VirtualProduction/LevelSnapshots/Source/LevelShapshotFilters/Public/PropertyBlueprintFunctionLibrary.h
@@ -23,6 +23,6 @@ public:
 
 	/* Loads the actor identified by Params. You can use this for advanced filter queries.*/
 	UFUNCTION(BlueprintPure, Category = "Level Snapshots")
-	AActor* LoadSnapshotActor(const FIsDeletedActorValidParams& Params);
+	static AActor* LoadSnapshotActor(const FIsDeletedActorValidParams& Params);
 	
 };
