LevelSnapshots: FilterByClass now also works in add/delete actors, not just modify.

This is the more expected behavior.

#rnx

#virtualized

diff --git a/Engine/Plugins/VirtualProduction/LevelSnapshots/Source/LevelShapshotFilters/Private/PropertyBlueprintFunctionLibrary.cpp b/Engine/Plugins/VirtualProduction/LevelSnapshots/Source/LevelShapshotFilters/Private/PropertyBlueprintFunctionLibrary.cpp
--- a/Engine/Plugins/VirtualProduction/LevelSnapshots/Source/LevelShapshotFilters/Private/PropertyBlueprintFunctionLibrary.cpp
+++ b/Engine/Plugins/VirtualProduction/LevelSnapshots/Source/LevelShapshotFilters/Private/PropertyBlueprintFunctionLibrary.cpp
@@ -18,3 +18,8 @@ AActor* UPropertyBlueprintFunctionLibrary::LoadSnapshotActor(const FIsDeletedAct
 {
 	return ensure(Params.HelperForDeserialization) ? Params.HelperForDeserialization(Params.SavedActorPath) : nullptr;
 }
+
+UClass* UPropertyBlueprintFunctionLibrary::GetActorClassFromDeletedParams(const FIsDeletedActorValidParams& Params)
+{
+	return ensure(Params.ClassGetter) ? Params.ClassGetter(Params.SavedActorPath) : nullptr;
+}
diff --git a/Engine/Plugins/VirtualProduction/LevelSnapshots/Source/LevelShapshotFilters/Public/LevelSnapshotFilterParams.h b/Engine/Plugins/VirtualProduction/LevelSnapshots/Source/LevelShapshotFilters/Public/LevelSnapshotFilterParams.h
--- a/Engine/Plugins/VirtualProduction/LevelSnapshots/Source/LevelShapshotFilters/Public/LevelSnapshotFilterParams.h
+++ b/Engine/Plugins/VirtualProduction/LevelSnapshots/Source/LevelShapshotFilters/Public/LevelSnapshotFilterParams.h
@@ -90,12 +90,20 @@ struct LEVELSNAPSHOTFILTERS_API FIsDeletedActorValidParams
 	FSoftObjectPath SavedActorPath;
 
 	/** Helper property for deserializing SavedActorPath. */
-	TFunction<AActor*(const FSoftObjectPath&)> HelperForDeserialization;
+	TFunction<AActor* (const FSoftObjectPath&)> HelperForDeserialization;
+
+	/** Helper for getting the class of the deleted actor */
+	TFunction<UClass* (const FSoftObjectPath&)> ClassGetter;
 
 	FIsDeletedActorValidParams() = default;
-	FIsDeletedActorValidParams(FSoftObjectPath SavedActorPath, TFunction<AActor*(const FSoftObjectPath&)> HelperForDeserialization)
+	FIsDeletedActorValidParams(
+		FSoftObjectPath SavedActorPath, 
+		TFunction<AActor*(const FSoftObjectPath&)> HelperForDeserialization,
+		TFunction<UClass* (const FSoftObjectPath&)> ClassGetter
+	)
 		: SavedActorPath(SavedActorPath)
 		, HelperForDeserialization(HelperForDeserialization)
+		, ClassGetter(ClassGetter)
 	{}
 };
 
diff --git a/Engine/Plugins/VirtualProduction/LevelSnapshots/Source/LevelShapshotFilters/Public/PropertyBlueprintFunctionLibrary.h b/Engine/Plugins/VirtualProduction/LevelSnapshots/Source/LevelShapshotFilters/Public/PropertyBlueprintFunctionLibrary.h
--- a/Engine/Plugins/VirtualProduction/LevelSnapshots/Source/LevelShapshotFilters/Public/PropertyBlueprintFunctionLibrary.h
+++ b/Engine/Plugins/VirtualProduction/LevelSnapshots/Source/LevelShapshotFilters/Public/PropertyBlueprintFunctionLibrary.h
@@ -24,5 +24,9 @@ public:
 	/* Loads the actor identified by Params. You can use this for advanced filter queries.*/
 	UFUNCTION(BlueprintPure, Category = "Level Snapshots")
 	static AActor* LoadSnapshotActor(const FIsDeletedActorValidParams& Params);
+
+	/* Finds the actor class from the given FIsDeletedActorValidParams. Use this for advanced filter queries */
+	UFUNCTION(BlueprintPure, Category = "Level Snapshots")
+	static UClass* GetActorClassFromDeletedParams(const FIsDeletedActorValidParams& Params);
 	
 };
diff --git a/Engine/Plugins/VirtualProduction/LevelSnapshots/Source/LevelSnapshots/Private/Data/LevelSnapshot.cpp b/Engine/Plugins/VirtualProduction/LevelSnapshots/Source/LevelSnapshots/Private/Data/LevelSnapshot.cpp
--- a/Engine/Plugins/VirtualProduction/LevelSnapshots/Source/LevelSnapshots/Private/Data/LevelSnapshot.cpp
+++ b/Engine/Plugins/VirtualProduction/LevelSnapshots/Source/LevelSnapshots/Private/Data/LevelSnapshot.cpp
@@ -180,6 +180,11 @@ TOptional<TNonNullPtr<AActor>> ULevelSnapshot::GetDeserializedActor(const FSoftO
 	return UE::LevelSnapshots::Private::GetDeserializedActor(OriginalActorPath, SerializedData, Cache, GetPackage());
 }
 
+TOptional<TNonNullPtr<UClass>> ULevelSnapshot::GetActorClass(const FSoftObjectPath& OriginalActorPath)
+{
+	return UE::LevelSnapshots::Private::GetActorClass(OriginalActorPath, SerializedData);
+}
+
 int32 ULevelSnapshot::GetNumSavedActors() const
 {
 	return SerializedData.ActorData.Num();
diff --git a/Engine/Plugins/VirtualProduction/LevelSnapshots/Source/LevelSnapshots/Private/Filtering/SnapshotSelectionBuilder.cpp b/Engine/Plugins/VirtualProduction/LevelSnapshots/Source/LevelSnapshots/Private/Filtering/SnapshotSelectionBuilder.cpp
--- a/Engine/Plugins/VirtualProduction/LevelSnapshots/Source/LevelSnapshots/Private/Filtering/SnapshotSelectionBuilder.cpp
+++ b/Engine/Plugins/VirtualProduction/LevelSnapshots/Source/LevelSnapshots/Private/Filtering/SnapshotSelectionBuilder.cpp
@@ -88,10 +88,20 @@ namespace UE::LevelSnapshots::Private
 		{
 			const EFilterResult::Type FilterResult = Filter->IsDeletedActorValid(
 				FIsDeletedActorValidParams(
+
+					// SavedActorPath
 					*RemovedActorsIt,
+
+					// HelperForDeserialization
 					[this](const FSoftObjectPath& ObjectPath)
 					{
 						return GetRawPointerOrNull(Snapshot->GetDeserializedActor(ObjectPath));
+					},
+
+					// ClassGetter
+					[this](const FSoftObjectPath& ObjectPath)
+					{
+						return GetRawPointerOrNull(Snapshot->GetActorClass(ObjectPath));
 					}
 				)
 			);
diff --git a/Engine/Plugins/VirtualProduction/LevelSnapshots/Source/LevelSnapshots/Private/Restoration/Actor/ActorUtil.h b/Engine/Plugins/VirtualProduction/LevelSnapshots/Source/LevelSnapshots/Private/Restoration/Actor/ActorUtil.h
--- a/Engine/Plugins/VirtualProduction/LevelSnapshots/Source/LevelSnapshots/Private/Restoration/Actor/ActorUtil.h
+++ b/Engine/Plugins/VirtualProduction/LevelSnapshots/Source/LevelSnapshots/Private/Restoration/Actor/ActorUtil.h
@@ -20,6 +20,9 @@ namespace UE::LevelSnapshots::Private
 	/** Allocates the actor, if no yet, and initialises it in the snapshot world. */
 	TOptional<TNonNullPtr<AActor>> GetDeserializedActor(const FSoftObjectPath& OriginalObjectPath, FWorldSnapshotData& WorldData, FSnapshotDataCache& Cache, UPackage* LocalisationSnapshotPackage, bool bWarnIfNotFound = true);
 
+	/** Gets the actor class from the original actor soft path */
+	TOptional<TNonNullPtr<UClass>> GetActorClass(const FSoftObjectPath& OriginalActorPath, FWorldSnapshotData& WorldData, bool bWarnIfNotFound = true);
+
 	TOptional<TNonNullPtr<AActor>> GetPreallocatedIfCached(const FSoftObjectPath& OriginalActorPath, const FSnapshotDataCache& Cache);
 
 	/** Gets the actor allocated in the snapshot world or allocates it and all of its subobjects. */
@@ -27,6 +30,7 @@ namespace UE::LevelSnapshots::Private
 
 	/** Restores selected properties into the actor and its subobjects: the actor existed before the snapshot was applied. */
 	void RestoreIntoExistingWorldActor(AActor* OriginalActor, FActorSnapshotData& ActorData, FWorldSnapshotData& WorldData, FSnapshotDataCache& Cache, UPackage* InLocalisationSnapshotPackage, const FPropertySelectionMap& SelectedProperties);
+
 	/** Restores all properties into the actor and its subobjects: the actor was recreated as part of the restoration process. */
 	void RestoreIntoRecreatedEditorWorldActor(AActor* OriginalActor, FActorSnapshotData& ActorData, FWorldSnapshotData& WorldData, FSnapshotDataCache& Cache, UPackage* InLocalisationSnapshotPackage, const FPropertySelectionMap& SelectedProperties);
 }
diff --git a/Engine/Plugins/VirtualProduction/LevelSnapshots/Source/LevelSnapshots/Private/Restoration/Actor/ActorUtil_TransientWorld.cpp b/Engine/Plugins/VirtualProduction/LevelSnapshots/Source/LevelSnapshots/Private/Restoration/Actor/ActorUtil_TransientWorld.cpp
--- a/Engine/Plugins/VirtualProduction/LevelSnapshots/Source/LevelSnapshots/Private/Restoration/Actor/ActorUtil_TransientWorld.cpp
+++ b/Engine/Plugins/VirtualProduction/LevelSnapshots/Source/LevelSnapshots/Private/Restoration/Actor/ActorUtil_TransientWorld.cpp
@@ -317,6 +317,28 @@ TOptional<TNonNullPtr<AActor>> UE::LevelSnapshots::Private::GetDeserializedActor
 	return Result;
 }
 
+TOptional<TNonNullPtr<UClass>> UE::LevelSnapshots::Private::GetActorClass(const FSoftObjectPath& OriginalActorPath, FWorldSnapshotData& WorldData, bool bWarnIfNotFound)
+{
+	const FActorSnapshotData* ActorData = WorldData.ActorData.Find(OriginalActorPath);
+
+	if (!ActorData)
+	{
+		UE_CLOG(bWarnIfNotFound, LogLevelSnapshots, Warning, TEXT("Failed to get the UClass for actor %s due to missing data"), *OriginalActorPath.ToString());
+		return {};
+	}
+
+	const FSoftClassPath ClassPath = GetClass(*ActorData, WorldData);
+	UClass* ActorClass = ClassPath.TryLoadClass<AActor>();
+
+	if (!ActorClass)
+	{
+		UE_LOG(LogLevelSnapshots, Warning, TEXT("Failed to load the UClass for actor %s."), *OriginalActorPath.ToString());
+		return {};
+	}
+
+	return TOptional<TNonNullPtr<UClass>>(ActorClass);
+}
+
 TOptional<TNonNullPtr<AActor>> UE::LevelSnapshots::Private::GetPreallocatedIfCached(const FSoftObjectPath& OriginalActorPath, const FSnapshotDataCache& Cache)
 {
 	if (const FActorSnapshotCache* ActorCache = Cache.ActorCache.Find(OriginalActorPath))
diff --git a/Engine/Plugins/VirtualProduction/LevelSnapshots/Source/LevelSnapshots/Public/Data/LevelSnapshot.h b/Engine/Plugins/VirtualProduction/LevelSnapshots/Source/LevelSnapshots/Public/Data/LevelSnapshot.h
--- a/Engine/Plugins/VirtualProduction/LevelSnapshots/Source/LevelSnapshots/Public/Data/LevelSnapshot.h
+++ b/Engine/Plugins/VirtualProduction/LevelSnapshots/Source/LevelSnapshots/Public/Data/LevelSnapshot.h
@@ -61,7 +61,10 @@ public:
 	
 	/** Given an actor path in the world, gets the equivalent actor from the snapshot. */
 	TOptional<TNonNullPtr<AActor>> GetDeserializedActor(const FSoftObjectPath& OriginalActorPath);
-	
+
+	/** Given an actor path in the world, gets its class from the snapshot. */
+	TOptional<TNonNullPtr<UClass>> GetActorClass(const FSoftObjectPath& OriginalActorPath);
+
 	/** Gets the number of actors saved by this snapshot. */
 	int32 GetNumSavedActors() const;
 	
