Make Handled/Unhandled [[nodiscard]]

diff --git a/Engine/Plugins/Animation/ControlRig/Source/ControlRigEditor/Private/Editor/SModularRigModel.cpp b/Engine/Plugins/Animation/ControlRig/Source/ControlRigEditor/Private/Editor/SModularRigModel.cpp
--- a/Engine/Plugins/Animation/ControlRig/Source/ControlRigEditor/Private/Editor/SModularRigModel.cpp
+++ b/Engine/Plugins/Animation/ControlRig/Source/ControlRigEditor/Private/Editor/SModularRigModel.cpp
@@ -1725,7 +1725,7 @@ FReply SModularRigModel::OnAcceptDrop(const FDragDropEvent& DragDropEvent, EItem
 			}
 		}
 
-		FReply::Handled();
+		return FReply::Handled();
 	}
 	else if(ModuleDragDropOperation)
 	{
diff --git a/Engine/Plugins/Editor/CurveEditorTools/Source/CurveEditorTools/Private/Tools/CurveEditorRetimeTool.cpp b/Engine/Plugins/Editor/CurveEditorTools/Source/CurveEditorTools/Private/Tools/CurveEditorRetimeTool.cpp
--- a/Engine/Plugins/Editor/CurveEditorTools/Source/CurveEditorTools/Private/Tools/CurveEditorRetimeTool.cpp
+++ b/Engine/Plugins/Editor/CurveEditorTools/Source/CurveEditorTools/Private/Tools/CurveEditorRetimeTool.cpp
@@ -166,7 +166,7 @@ FReply FCurveEditorRetimeTool::OnMouseMove(TSharedRef<SWidget> OwningWidget, con
 	TSharedPtr<FCurveEditor> CurveEditor = WeakCurveEditor.Pin();
 	if (!CurveEditor)
 	{
-		FReply::Unhandled();
+		return FReply::Unhandled();
 	}
 
 	for (FCurveEditorRetimeAnchor& Anchor : RetimeData->RetimingAnchors)
diff --git a/Engine/Plugins/VirtualProduction/DynamicMaterial/Source/DynamicMaterialEditor/Private/UI/Widgets/SDMMaterialWizard.cpp b/Engine/Plugins/VirtualProduction/DynamicMaterial/Source/DynamicMaterialEditor/Private/UI/Widgets/SDMMaterialWizard.cpp
--- a/Engine/Plugins/VirtualProduction/DynamicMaterial/Source/DynamicMaterialEditor/Private/UI/Widgets/SDMMaterialWizard.cpp
+++ b/Engine/Plugins/VirtualProduction/DynamicMaterial/Source/DynamicMaterialEditor/Private/UI/Widgets/SDMMaterialWizard.cpp
@@ -589,7 +589,7 @@ FReply SDMMaterialWizard::NewTemplateAccept_OnClick()
 {
 	if (CurrentPreset.IsNone())
 	{
-		FReply::Handled();
+		return FReply::Handled();
 	}
 
 	TSharedPtr<SDMMaterialDesigner> DesignerWidget = GetDesignerWidget();
diff --git a/Engine/Plugins/VirtualProduction/LevelSnapshots/Source/LevelSnapshotsEditor/Private/Views/Results/SLevelSnapshotsEditorResults.cpp b/Engine/Plugins/VirtualProduction/LevelSnapshots/Source/LevelSnapshotsEditor/Private/Views/Results/SLevelSnapshotsEditorResults.cpp
--- a/Engine/Plugins/VirtualProduction/LevelSnapshots/Source/LevelSnapshotsEditor/Private/Views/Results/SLevelSnapshotsEditorResults.cpp
+++ b/Engine/Plugins/VirtualProduction/LevelSnapshots/Source/LevelSnapshotsEditor/Private/Views/Results/SLevelSnapshotsEditorResults.cpp
@@ -366,7 +366,7 @@ FReply SLevelSnapshotsEditorResults::OnClickApplyToWorld()
 {
 	if (!ensure(EditorDataPtr.IsValid()))
 	{
-		FReply::Handled();
+		return FReply::Handled();
 	}
 
 	const TOptional<ULevelSnapshot*> ActiveLevelSnapshot = EditorDataPtr->GetActiveSnapshot();
diff --git a/Engine/Source/Editor/Persona/Private/AnimAssetFindReplace.cpp b/Engine/Source/Editor/Persona/Private/AnimAssetFindReplace.cpp
--- a/Engine/Source/Editor/Persona/Private/AnimAssetFindReplace.cpp
+++ b/Engine/Source/Editor/Persona/Private/AnimAssetFindReplace.cpp
@@ -94,7 +94,7 @@ public:
 								TGuardValue<bool> GuardValue(bSettingTextFromSearchItem, true);
 								SearchBox->SetText(FText::FromString(*SelectedItems[0].Get()));
 								MenuAnchor->SetIsOpen(false);
-								FReply::Handled();
+								return FReply::Handled();
 							}
 						}
 						return FReply::Unhandled();
diff --git a/Engine/Source/Runtime/SlateCore/Public/Input/Reply.h b/Engine/Source/Runtime/SlateCore/Public/Input/Reply.h
--- a/Engine/Source/Runtime/SlateCore/Public/Input/Reply.h
+++ b/Engine/Source/Runtime/SlateCore/Public/Input/Reply.h
@@ -230,7 +230,7 @@ public:
 	/**
 	 * An event should return a FReply::Handled() to let the system know that an event was handled.
 	 */
-	static FReply Handled( )
+	[[nodiscard]] static FReply Handled( )
 	{
 		return FReply(true);
 	}
@@ -238,7 +238,7 @@ public:
 	/**
 	 * An event should return a FReply::Unhandled() to let the system know that an event was unhandled.
 	 */
-	static FReply Unhandled( )
+	[[nodiscard]] static FReply Unhandled( )
 	{
 		return FReply(false);
 	}
