Level Snapshots: Fix crash when deciding whether a property should be saved, if the property non-reflected.

For example, a crash would occur snapshotting ACEClonerActor.
The reason was that a niagara component would call SerializeTaggedProperties the structs in a TArray<TSharedRef<FMyStruct>>.
This would still invoke ShouldSkipProperty on the properties of thy FMyStruct properties but the container pointer we were passing to Identical_InContainer was not the UObject but the unreflected TArray.

This change also contains some nullptr checks for Niagara. Level Snapshots does not currently restore the niagara actors correctly. These changes prevent the editor from crashing on restore.

diff --git a/Engine/Plugins/FX/Niagara/Source/Niagara/Private/NiagaraParameterDefinitionsSubscriber.cpp b/Engine/Plugins/FX/Niagara/Source/Niagara/Private/NiagaraParameterDefinitionsSubscriber.cpp
--- a/Engine/Plugins/FX/Niagara/Source/Niagara/Private/NiagaraParameterDefinitionsSubscriber.cpp
+++ b/Engine/Plugins/FX/Niagara/Source/Niagara/Private/NiagaraParameterDefinitionsSubscriber.cpp
@@ -259,7 +259,10 @@ void INiagaraParameterDefinitionsSubscriber::SynchronizeWithParameterDefinitions
 	// Synchronize source scripts.
 	for (UNiagaraScriptSourceBase* SourceScript : GetAllSourceScripts())
 	{
-		SourceScript->SynchronizeGraphParametersWithParameterDefinitions(TargetDefinitions, AllDefinitions, DefinitionParameterIds, this, Args);
+		if (ensure(SourceScript))
+		{
+			SourceScript->SynchronizeGraphParametersWithParameterDefinitions(TargetDefinitions, AllDefinitions, DefinitionParameterIds, this, Args);
+		}
 	}
 
 	// Synchronize editor only script variables.
diff --git a/Engine/Plugins/FX/Niagara/Source/Niagara/Private/NiagaraSystem.cpp b/Engine/Plugins/FX/Niagara/Source/Niagara/Private/NiagaraSystem.cpp
--- a/Engine/Plugins/FX/Niagara/Source/Niagara/Private/NiagaraSystem.cpp
+++ b/Engine/Plugins/FX/Niagara/Source/Niagara/Private/NiagaraSystem.cpp
@@ -674,9 +674,8 @@ void UNiagaraSystem::UpdateSystemAfterLoad()
 	}
 
 #if WITH_EDITORONLY_DATA
-	if (SystemSpawnScript && SystemUpdateScript && !GetOutermost()->bIsCookedForEditor)
+	if (SystemSpawnScript && SystemUpdateScript && !GetOutermost()->bIsCookedForEditor && ensure(SystemSpawnScript->GetLatestSource() != nullptr))
 	{
-		ensure(SystemSpawnScript->GetLatestSource() != nullptr);
 		ensure(SystemSpawnScript->GetLatestSource() == SystemUpdateScript->GetLatestSource());
 		SystemSpawnScript->GetLatestSource()->OnChanged().AddUObject(this, &UNiagaraSystem::GraphSourceChanged);
 	}
diff --git a/Engine/Plugins/VirtualProduction/LevelSnapshots/Source/LevelSnapshots/Private/Archive/TakeWorldObjectSnapshotArchive.cpp b/Engine/Plugins/VirtualProduction/LevelSnapshots/Source/LevelSnapshots/Private/Archive/TakeWorldObjectSnapshotArchive.cpp
--- a/Engine/Plugins/VirtualProduction/LevelSnapshots/Source/LevelSnapshots/Private/Archive/TakeWorldObjectSnapshotArchive.cpp
+++ b/Engine/Plugins/VirtualProduction/LevelSnapshots/Source/LevelSnapshots/Private/Archive/TakeWorldObjectSnapshotArchive.cpp
@@ -39,6 +39,15 @@ bool UE::LevelSnapshots::Private::FTakeWorldObjectSnapshotArchive::ShouldSkipPro
 
 	if (!bSuperWantsToSkip)
 	{
+		// The GetSerializedObject() may have overriden UObject::Serialize and may be serializing some custom, non-reflected data.
+		// For example, the object may have a non-reflected FMyStruct but manually call SerializeTaggedProperties on the FMyStruct instance.
+		// In that case, we cannot call Identical_InContainer with GetSerializedObject() as owning object is not necessarily the container.
+		// Solution: Just allow snapshotting all custom properties. Saving the additional property data is the safest thing to do.
+		if (!bIsSerializingScriptProperties)
+		{
+			return false;
+		}
+		
 		const bool bIsRootProperty = [this]()
 		{
 			const FArchiveSerializedPropertyChain* PropertyChain = GetSerializedPropertyChain();
diff --git a/Engine/Plugins/VirtualProduction/LevelSnapshots/Source/LevelSnapshots/Private/Archive/TakeWorldObjectSnapshotArchive.h b/Engine/Plugins/VirtualProduction/LevelSnapshots/Source/LevelSnapshots/Private/Archive/TakeWorldObjectSnapshotArchive.h
--- a/Engine/Plugins/VirtualProduction/LevelSnapshots/Source/LevelSnapshots/Private/Archive/TakeWorldObjectSnapshotArchive.h
+++ b/Engine/Plugins/VirtualProduction/LevelSnapshots/Source/LevelSnapshots/Private/Archive/TakeWorldObjectSnapshotArchive.h
@@ -21,6 +21,8 @@ namespace UE::LevelSnapshots::Private
 
 		//~ Begin FSnapshotArchive Interface
 		virtual bool ShouldSkipProperty(const FProperty* InProperty) const override;
+		virtual void MarkScriptSerializationStart(const UObject* Obj) override { bIsSerializingScriptProperties = true; }
+		virtual void MarkScriptSerializationEnd(const UObject* Obj) override { bIsSerializingScriptProperties = false; }
 		//~ End FSnapshotArchive Interface
 
 	protected:
@@ -34,5 +36,11 @@ namespace UE::LevelSnapshots::Private
 		FTakeWorldObjectSnapshotArchive(FObjectSnapshotData& InObjectData, FWorldSnapshotData& InSharedData, UObject* InOriginalObject);
 
 		UObject* Archetype;
+
+		/**
+		 * Whether currently serializing UPROPERTY()s on the SerializedObject.
+		 * This is used to detect whether some custom properties are being serialized.
+		 */
+		bool bIsSerializingScriptProperties = false;
 	};
 }
\ No newline at end of file
