Add a console command that will allow users to capture level snapshot from a game world.

[RN] Added a new console command called LevelSnapshots.TakeLevelSnapshot that allows users to capture a level snapshot from a game world.
#Synced-CL 49278929

diff --git a/Engine/Plugins/VirtualProduction/LevelSnapshots/Source/LevelSnapshotsEditor/Private/Util/TakeSnapshotUtil.cpp b/Engine/Plugins/VirtualProduction/LevelSnapshots/Source/LevelSnapshotsEditor/Private/Util/TakeSnapshotUtil.cpp
--- a/Engine/Plugins/VirtualProduction/LevelSnapshots/Source/LevelSnapshotsEditor/Private/Util/TakeSnapshotUtil.cpp
+++ b/Engine/Plugins/VirtualProduction/LevelSnapshots/Source/LevelSnapshotsEditor/Private/Util/TakeSnapshotUtil.cpp
@@ -25,15 +25,37 @@
 #include "UObject/SavePackage.h"
 #include "Widgets/Notifications/SNotificationList.h"
 
+#include "Engine/GameEngine.h"
+
 #define LOCTEXT_NAMESPACE "LevelSnapshotEditorLibrary"
 
 namespace UE::LevelSnapshots::Editor
 {
 	TAutoConsoleVariable<bool> CVarSaveAsync(TEXT("LevelSnapshots.SnapshotFormUsesSaveAsync"), true, TEXT("Set whether the Take Snapshot Form, opened by pressing the toolbar icon, saves snapshots using the SAVE_Async flag."));
+	TAutoConsoleVariable<bool> CVarIgnorePIEWorld(TEXT("LevelSnapshots.IgnorePIEWorld"), true, TEXT("Set whether the Take Snapshot will ignore the PIE world and use the editor world context.  If this is false, then the PIE world will be used to grab a snapshot."));
 	
+	UWorld* GetWorld(bool bIgnorePIE)
+	{
+		// Get the current world.
+		if (GIsEditor)
+		{
+			FWorldContext* PIEWorldContext = GEditor->GetPIEWorldContext();
+			if (!bIgnorePIE && PIEWorldContext)
+			{
+				return PIEWorldContext->World();
+			}
+			return GEditor->GetEditorWorldContext().World();
+		}
+		else if (UGameEngine* GameEngine = Cast<UGameEngine>(GEngine))
+		{
+			return GameEngine->GetGameWorld();
+		}
+		return nullptr;
+	}
+
 	static void HandleFormReply(const FText& InDescription, const FString& SnapshotName)
 	{
-		UWorld* World = ULevelSnapshotsEditorData::GetEditorWorld();
+		UWorld* World = GetWorld(CVarIgnorePIEWorld.GetValueOnAnyThread());
 		if (!ensure(World))
 		{
 			return;
@@ -206,7 +228,7 @@ void SnapshotEditor::GenerateThumbnailForSnapshotAsset(ULevelSnapshot* Snapshot)
 	TArray<FAssetData> SnapshotAssetData;
 	AssetRegistry.GetAssetsByPackageName(FName(*Snapshot->GetPackage()->GetPathName()), SnapshotAssetData);
 	
-	if (ensureMsgf(SnapshotAssetData.Num(), TEXT("Failed to find asset data for asset we just saved. Investigate!")))
+	if (GEditor && ensureMsgf(SnapshotAssetData.Num(), TEXT("Failed to find asset data for asset we just saved. Investigate!")))
 	{
 		// Copied from FAssetFileContextMenu::ExecuteCaptureThumbnail
 		FViewport* Viewport = GEditor->GetActiveViewport();
@@ -228,4 +250,32 @@ void SnapshotEditor::GenerateThumbnailForSnapshotAsset(ULevelSnapshot* Snapshot)
 	}
 }
 
+
+namespace UE::LevelSnapshots::Private
+{
+
+FAutoConsoleCommandWithWorldAndArgs TakeLevelSnaphot(
+	TEXT("LevelSnapshots.TakeLevelSnapshot"),
+	TEXT("Take a new level snapshot of the current world.")
+	TEXT("Arg1 - Name of the level snapshot asset.\n")
+	TEXT("Arg2 - Folder path to store the snapshot, e.g. /Game/LevelSnapshots. \n"),
+	FConsoleCommandWithWorldAndArgsDelegate::CreateLambda(
+		[](const TArray<FString>& Args, UWorld* World)
+		{
+			if (Args.Num() == 2)
+			{
+				FString Description(TEXT("Snapshot taken from console command."));
+				
+				SnapshotEditor::TakeLevelSnapshotAndSaveToDisk(World, Args[0], Args[1], Description, true, false);
+			}
+			else
+			{
+				UE_LOG(LogLevelSnapshots, Warning, TEXT("Incorrect number of arguments for the snapshot command."));				
+			}
+		}
+	)
+);
+
+} // namespace UE::LevelSnapshots::Private
+
 #undef LOCTEXT_NAMESPACE
diff --git a/Engine/Plugins/VirtualProduction/LevelSnapshots/Source/LevelSnapshotsEditor/Private/Util/TakeSnapshotUtil.h b/Engine/Plugins/VirtualProduction/LevelSnapshots/Source/LevelSnapshotsEditor/Private/Util/TakeSnapshotUtil.h
--- a/Engine/Plugins/VirtualProduction/LevelSnapshots/Source/LevelSnapshotsEditor/Private/Util/TakeSnapshotUtil.h
+++ b/Engine/Plugins/VirtualProduction/LevelSnapshots/Source/LevelSnapshotsEditor/Private/Util/TakeSnapshotUtil.h
@@ -8,6 +8,14 @@ class ULevelSnapshot;
 
 namespace SnapshotEditor
 {
+	/**
+	 * Get the world to use for taking a level snapshot.
+	 * @param bIgnorePIE  if true and we are in a PIE world, we ignore the PIE world and use the editor world context.
+	 *
+	 * @return The current UWorld.
+	 */
+	UWorld* GetWorld(bool bIgnorePIE);
+
 	/**
 	 * Shows the window for taking snapshots, if enabled by the user settings. Otherwise just takes a snapshot.
 	 */
